Fixed deserialization of LazyObjectPtr in Multi-User when the pointed object is null

diff --git a/Engine/Plugins/Developer/Concert/ConcertSync/ConcertSyncCore/Source/ConcertSyncCore/Private/ConcertSyncArchives.cpp b/Engine/Plugins/Developer/Concert/ConcertSync/ConcertSyncCore/Source/ConcertSyncCore/Private/ConcertSyncArchives.cpp
--- a/Engine/Plugins/Developer/Concert/ConcertSync/ConcertSyncCore/Source/ConcertSyncCore/Private/ConcertSyncArchives.cpp
+++ b/Engine/Plugins/Developer/Concert/ConcertSync/ConcertSyncCore/Source/ConcertSyncCore/Private/ConcertSyncArchives.cpp
@@ -256,7 +256,7 @@ FArchive& FConcertSyncObjectReader::operator<<(FLazyObjectPtr& LazyObjectPtr)
 
 	// if the resolved object already has an associated Guid, use that instead of the saved one
 	// otherwise used the saved guid since it should refer to the Obj path once its state get applied.
-	FUniqueObjectGuid ObjectGuid(Obj);
+	FUniqueObjectGuid ObjectGuid = Obj != nullptr ? FUniqueObjectGuid(Obj) : FUniqueObjectGuid();
 	LazyObjectPtr = ObjectGuid.IsValid() ? ObjectGuid : SavedObjectGuid;
 	// technically the saved object guid should be the same as the resolved object guid if any.
 	ensure(!ObjectGuid.IsValid() || ObjectGuid == SavedObjectGuid);
