Fixing deprecation warnings

diff --git a/Engine/Plugins/Developer/Concert/ConcertSync/ConcertSyncCore/Source/ConcertSyncCore/Private/ConcertActivityLedger.cpp b/Engine/Plugins/Developer/Concert/ConcertSync/ConcertSyncCore/Source/ConcertSyncCore/Private/ConcertActivityLedger.cpp
--- a/Engine/Plugins/Developer/Concert/ConcertSync/ConcertSyncCore/Source/ConcertSyncCore/Private/ConcertActivityLedger.cpp
+++ b/Engine/Plugins/Developer/Concert/ConcertSync/ConcertSyncCore/Source/ConcertSyncCore/Private/ConcertActivityLedger.cpp
@@ -57,7 +57,7 @@ namespace ConcertActivityLedgerUtil
 		Ar.SerializeIntPacked(UncompressedActivitySize);
 		if (UncompressedActivitySize > 0)
 		{
-			Ar.SerializeCompressed(UncompressedActivity.GetData(), UncompressedActivitySize, COMPRESS_ZLIB);
+			Ar.SerializeCompressed(UncompressedActivity.GetData(), UncompressedActivitySize, NAME_Zlib);
 		}
 
 		// Serialize the footer so we know we didn't crash mid-write
@@ -129,7 +129,7 @@ namespace ConcertActivityLedgerUtil
 		UncompressedActivity.AddZeroed(UncompressedActivitySize);
 		if (UncompressedActivitySize > 0)
 		{
-			Ar.SerializeCompressed(UncompressedActivity.GetData(), UncompressedActivitySize, COMPRESS_ZLIB);
+			Ar.SerializeCompressed(UncompressedActivity.GetData(), UncompressedActivitySize, NAME_Zlib);
 		}
 
 		// Read the raw transaction data
diff --git a/Engine/Plugins/Developer/Concert/ConcertSync/ConcertSyncCore/Source/ConcertSyncCore/Private/ConcertPackageLedger.cpp b/Engine/Plugins/Developer/Concert/ConcertSync/ConcertSyncCore/Source/ConcertSyncCore/Private/ConcertPackageLedger.cpp
--- a/Engine/Plugins/Developer/Concert/ConcertSync/ConcertSyncCore/Source/ConcertSyncCore/Private/ConcertPackageLedger.cpp
+++ b/Engine/Plugins/Developer/Concert/ConcertSync/ConcertSyncCore/Source/ConcertSyncCore/Private/ConcertPackageLedger.cpp
@@ -47,7 +47,7 @@ bool WritePackage(const FConcertPackageInfo& InPackageInfo, const TArray<uint8>&
 	Ar.SerializeIntPacked(UncompressedPackageSize);
 	if (UncompressedPackageSize > 0)
 	{
-		Ar.SerializeCompressed((uint8*)InPackageData.GetData(), UncompressedPackageSize, COMPRESS_ZLIB);
+		Ar.SerializeCompressed((uint8*)InPackageData.GetData(), UncompressedPackageSize, NAME_Zlib);
 	}
 
 	// Serialize the footer so we know we didn't crash mid-write
@@ -98,7 +98,7 @@ bool ReadPackage(const TArray<uint8>& InSerializedPackageData, FConcertPackageIn
 		OutPackageData->AddZeroed(UncompressedPackageSize);
 		if (UncompressedPackageSize > 0)
 		{
-			Ar.SerializeCompressed(OutPackageData->GetData(), UncompressedPackageSize, COMPRESS_ZLIB);
+			Ar.SerializeCompressed(OutPackageData->GetData(), UncompressedPackageSize, NAME_Zlib);
 		}
 	}
 
diff --git a/Engine/Plugins/Developer/Concert/ConcertSync/ConcertSyncCore/Source/ConcertSyncCore/Private/ConcertTransactionLedger.cpp b/Engine/Plugins/Developer/Concert/ConcertSync/ConcertSyncCore/Source/ConcertSyncCore/Private/ConcertTransactionLedger.cpp
--- a/Engine/Plugins/Developer/Concert/ConcertSync/ConcertSyncCore/Source/ConcertSyncCore/Private/ConcertTransactionLedger.cpp
+++ b/Engine/Plugins/Developer/Concert/ConcertSync/ConcertSyncCore/Source/ConcertSyncCore/Private/ConcertTransactionLedger.cpp
@@ -55,7 +55,7 @@ bool WriteTransaction(const FStructOnScope& InTransaction, TArray<uint8>& OutSer
 	Ar.SerializeIntPacked(UncompressedTransactionSize);
 	if (UncompressedTransactionSize > 0)
 	{
-		Ar.SerializeCompressed(UncompressedTransaction.GetData(), UncompressedTransactionSize, COMPRESS_ZLIB);
+		Ar.SerializeCompressed(UncompressedTransaction.GetData(), UncompressedTransactionSize, NAME_Zlib);
 	}
 
 	// Serialize the footer so we know we didn't crash mid-write
@@ -126,7 +126,7 @@ bool ReadTransaction(const TArray<uint8>& InSerializedTransactionData, FStructOn
 	UncompressedTransaction.AddZeroed(UncompressedTransactionSize);
 	if (UncompressedTransactionSize > 0)
 	{
-		Ar.SerializeCompressed(UncompressedTransaction.GetData(), UncompressedTransactionSize, COMPRESS_ZLIB);
+		Ar.SerializeCompressed(UncompressedTransaction.GetData(), UncompressedTransactionSize, NAME_Zlib);
 	}
 
 	// Read the raw transaction data
