Added Polys, Model and BrushBuilder to the default whitelist transaction filter in Concert

diff --git a/Engine/Plugins/Developer/Concert/ConcertSync/ConcertSyncCore/Config/Engine.ini b/Engine/Plugins/Developer/Concert/ConcertSync/ConcertSyncCore/Config/Engine.ini
--- a/Engine/Plugins/Developer/Concert/ConcertSync/ConcertSyncCore/Config/Engine.ini
+++ b/Engine/Plugins/Developer/Concert/ConcertSync/ConcertSyncCore/Config/Engine.ini
@@ -5,6 +5,9 @@
 +IncludeObjectClassFilters=(ObjectClass=Class'/Script/Engine.LevelStreaming',ObjectOuterClass=Class'/Script/Engine.World')
 +IncludeObjectClassFilters=(ObjectClass=Class'/Script/Engine.Actor',ObjectOuterClass=Class'/Script/Engine.World')
 +IncludeObjectClassFilters=(ObjectClass=Class'/Script/Engine.ActorComponent',ObjectOuterClass=Class'/Script/Engine.World')
++IncludeObjectClassFilters=(ObjectClass=Class'/Script/Engine.Polys',ObjectOuterClass=Class'/Script/Engine.World')
++IncludeObjectClassFilters=(ObjectClass=Class'/Script/Engine.Model',ObjectOuterClass=Class'/Script/Engine.World')
++IncludeObjectClassFilters=(ObjectClass=Class'/Script/Engine.BrushBuilder',ObjectOuterClass=Class'/Script/Engine.World')
 +IncludeObjectClassFilters=(ObjectClass=Class'/Script/Engine.Actor',ObjectOuterClass=Class'/Script/MovieScene.MovieSceneSignedObject')
 +IncludeObjectClassFilters=(ObjectClass=Class'/Script/Engine.ActorComponent',ObjectOuterClass=Class'/Script/MovieScene.MovieSceneSignedObject')
 +IncludeObjectClassFilters=(ObjectClass=Class'/Script/Engine.BookmarkBase',ObjectOuterClass=None)
