Fixed new foliage types failing to propagate via MU

They have to have RF_Transactional set when being constructed in order to be caught by the transaction correctly

#rnx

diff --git a/Engine/Plugins/Developer/Concert/ConcertSync/ConcertSyncCore/Config/Engine.ini b/Engine/Plugins/Developer/Concert/ConcertSync/ConcertSyncCore/Config/Engine.ini
--- a/Engine/Plugins/Developer/Concert/ConcertSync/ConcertSyncCore/Config/Engine.ini
+++ b/Engine/Plugins/Developer/Concert/ConcertSync/ConcertSyncCore/Config/Engine.ini
@@ -12,6 +12,7 @@
 +IncludeObjectClassFilters=(ObjectClass=Class'/Script/Engine.ActorComponent',ObjectOuterClass=Class'/Script/MovieScene.MovieSceneSignedObject')
 +IncludeObjectClassFilters=(ObjectClass=Class'/Script/Engine.BookmarkBase',ObjectOuterClass=None)
 +IncludeObjectClassFilters=(ObjectClass=Class'/Script/Engine.MaterialParameterCollection',ObjectOuterClass=None)
++IncludeObjectClassFilters=(ObjectClass=Class'/Script/Foliage.FoliageType',ObjectOuterClass=Class'/Script/Engine.World')
 +IncludeObjectClassFilters=(ObjectClass=Class'/Script/MovieScene.MovieSceneSignedObject',ObjectOuterClass=None)
 +IncludeObjectClassFilters=(ObjectClass=Class'/Script/MovieScene.MovieSceneFolder',ObjectOuterClass=None)
 +IncludeObjectClassFilters=(ObjectClass=Class'/Script/Composure.CompositingElementPass',ObjectOuterClass=None)
diff --git a/Engine/Source/Runtime/Foliage/Private/FoliageTypeObject.cpp b/Engine/Source/Runtime/Foliage/Private/FoliageTypeObject.cpp
--- a/Engine/Source/Runtime/Foliage/Private/FoliageTypeObject.cpp
+++ b/Engine/Source/Runtime/Foliage/Private/FoliageTypeObject.cpp
@@ -19,7 +19,7 @@ void FFoliageTypeObject::RefreshInstance()
 		if (FoliageTypeAsBP->ParentClass->IsChildOf<UFoliageType_InstancedStaticMesh>())
 		{
 			bIsAsset = false;
-			TypeInstance = NewObject<UFoliageType_InstancedStaticMesh>(GetTransientPackage(), *FoliageTypeAsBP->GeneratedClass);
+			TypeInstance = NewObject<UFoliageType_InstancedStaticMesh>(GetTransientPackage(), *FoliageTypeAsBP->GeneratedClass, NAME_None, RF_Transactional);
 		}
 	}
 }
diff --git a/Engine/Source/Runtime/Foliage/Private/InstancedFoliage.cpp b/Engine/Source/Runtime/Foliage/Private/InstancedFoliage.cpp
--- a/Engine/Source/Runtime/Foliage/Private/InstancedFoliage.cpp
+++ b/Engine/Source/Runtime/Foliage/Private/InstancedFoliage.cpp
@@ -2231,14 +2231,15 @@ FFoliageMeshInfo* AInstancedFoliageActor::AddMesh(UStaticMesh* InMesh, UFoliageT
 	if (DefaultSettings)
 	{
 		// TODO: Can't we just use this directly?
-		Settings = DuplicateObject<UFoliageType_InstancedStaticMesh>(DefaultSettings, this);
+		FObjectDuplicationParameters DuplicationParameters(const_cast<UFoliageType_InstancedStaticMesh*>(DefaultSettings), this);
+		DuplicationParameters.ApplyFlags = RF_Transactional;
+		Settings = CastChecked<UFoliageType_InstancedStaticMesh>(StaticDuplicateObjectEx(DuplicationParameters));
 	}
 	else
 #endif
 	{
-		Settings = NewObject<UFoliageType_InstancedStaticMesh>(this);
+		Settings = NewObject<UFoliageType_InstancedStaticMesh>(this, NAME_None, RF_Transactional);
 	}
-	Settings->SetFlags(RF_Transactional);
 	Settings->Mesh = InMesh;
 
 	FFoliageMeshInfo* MeshInfo = AddMesh(Settings);
