Remove Composure soft reference from Concert.

diff --git a/Engine/Plugins/Developer/Concert/ConcertSync/ConcertSyncCore/Config/Engine.ini b/Engine/Plugins/Developer/Concert/ConcertSync/ConcertSyncCore/Config/Engine.ini
--- a/Engine/Plugins/Developer/Concert/ConcertSync/ConcertSyncCore/Config/Engine.ini
+++ b/Engine/Plugins/Developer/Concert/ConcertSync/ConcertSyncCore/Config/Engine.ini
@@ -15,11 +15,13 @@
 +IncludeObjectClassFilters=(ObjectClass=Class'/Script/Foliage.FoliageType',ObjectOuterClass=Class'/Script/Engine.World')
 +IncludeObjectClassFilters=(ObjectClass=Class'/Script/MovieScene.MovieSceneSignedObject',ObjectOuterClass=None)
 +IncludeObjectClassFilters=(ObjectClass=Class'/Script/MovieScene.MovieSceneFolder',ObjectOuterClass=None)
-+IncludeObjectClassFilters=(ObjectClass=Class'/Script/Composure.CompositingElementPass',ObjectOuterClass=None)
-+IncludeObjectClassFilters=(ObjectClass=Class'/Script/Composure.ComposurePostProcessPassPolicy',ObjectOuterClass=None)
 +ExcludeTransactionClassFilters=(ObjectClass=Class'/Script/Landscape.LandscapeProxy',ObjectOuterClass=None)
 +ExcludeTransactionClassFilters=(ObjectClass=Class'/Script/Landscape.LandscapeHeightfieldCollisionComponent',ObjectOuterClass=None)
 +ExcludeTransactionClassFilters=(ObjectClass=Class'/Script/Landscape.LandscapeComponent',ObjectOuterClass=None)
 +AllowedTransientProperties=Property'/Script/Engine.Actor:bHiddenEdTemporary'
 +AllowedTransientProperties=Property'/Script/Engine.Actor:bHiddenEdLevel'
-+AllowedTransientProperties=Property'/Script/Engine.Level:bIsVisible'
\ No newline at end of file
++AllowedTransientProperties=Property'/Script/Engine.Level:bIsVisible'
+
+; If you wish to use Composure in Multi-user, you may enable those options in your project settings.
+;+IncludeObjectClassFilters=(ObjectClass=Class'/Script/Composure.CompositingElementPass',ObjectOuterClass=None)
+;+IncludeObjectClassFilters=(ObjectClass=Class'/Script/Composure.ComposurePostProcessPassPolicy',ObjectOuterClass=None)
\ No newline at end of file
