Edigrate cl 9600687:
Fixed EnumeratePackageNamesWithLiveTransactions to enumerate just packages with live transactions as opposed to packages with transactions
Fixed applying live transactions to packages properly

diff --git a/Engine/Plugins/Developer/Concert/ConcertSync/ConcertSyncClient/Source/ConcertSyncClient/Private/ConcertClientWorkspace.cpp b/Engine/Plugins/Developer/Concert/ConcertSync/ConcertSyncClient/Source/ConcertSyncClient/Private/ConcertClientWorkspace.cpp
--- a/Engine/Plugins/Developer/Concert/ConcertSync/ConcertSyncClient/Source/ConcertSyncClient/Private/ConcertClientWorkspace.cpp
+++ b/Engine/Plugins/Developer/Concert/ConcertSync/ConcertSyncClient/Source/ConcertSyncClient/Private/ConcertClientWorkspace.cpp
@@ -524,6 +524,9 @@ void FConcertClientWorkspace::SaveLiveTransactionsToPackages()
 			UPackage* Package = LoadPackage(nullptr, *PackageNameStr, LOAD_None);
 			if (Package)
 			{
+				// Load package will queue live transaction, process them before saving the file
+				TransactionManager->ProcessPending();
+
 				UWorld* World = UWorld::FindWorldInPackage(Package);
 				FString PackageFilename;
 				if (!FPackageName::DoesPackageExist(PackageNameStr, nullptr, &PackageFilename))
diff --git a/Engine/Plugins/Developer/Concert/ConcertSync/ConcertSyncCore/Source/ConcertSyncCore/Private/ConcertSyncSessionDatabase.cpp b/Engine/Plugins/Developer/Concert/ConcertSync/ConcertSyncCore/Source/ConcertSyncCore/Private/ConcertSyncSessionDatabase.cpp
--- a/Engine/Plugins/Developer/Concert/ConcertSync/ConcertSyncCore/Source/ConcertSyncCore/Private/ConcertSyncSessionDatabase.cpp
+++ b/Engine/Plugins/Developer/Concert/ConcertSync/ConcertSyncCore/Source/ConcertSyncCore/Private/ConcertSyncSessionDatabase.cpp
@@ -353,6 +353,7 @@ public:
 		PREPARE_STATEMENT(Statement_UnmapPackageNameIdsForTransactionEventId);
 		PREPARE_STATEMENT(Statement_GetTransactionEventIdsForPackageNameId);
 		PREPARE_STATEMENT(Statement_GetTransactionEventIdsInRangeForPackageNameId);
+		PREPARE_STATEMENT(Statement_GetPackageNameIdsMaxTransactionId);
 		PREPARE_STATEMENT(Statement_GetPackageNameIdsWithTransactions);
 		PREPARE_STATEMENT(Statement_GetPackageNameIdsForTransactionEventId);
 
@@ -1004,6 +1005,22 @@ public:
 		}) != INDEX_NONE;
 	}
 
+	/** Get the max transaction_event_id for each package_name_ids from package_transactions */
+	SQLITE_PREPARED_STATEMENT_COLUMNS_ONLY(FGetPackageNameIdsMaxTransactionId, "SELECT package_name_id, MAX(transaction_event_id) FROM package_transactions GROUP BY package_name_id;", SQLITE_PREPARED_STATEMENT_COLUMNS(int64, int64));
+	FGetPackageNameIdsMaxTransactionId Statement_GetPackageNameIdsMaxTransactionId;
+	bool GetPackageNameIdsMaxTransactionId(TFunctionRef<ESQLitePreparedStatementExecuteRowResult(int64, int64)> InCallback)
+	{
+		return Statement_GetPackageNameIdsMaxTransactionId.Execute([&InCallback](const FGetPackageNameIdsMaxTransactionId& InStatement)
+		{
+			int64 PackageNameId = 0, MaxTransactionEventId = 0;
+			if (InStatement.GetColumnValues(PackageNameId, MaxTransactionEventId))
+			{
+				return InCallback(PackageNameId, MaxTransactionEventId);
+			}
+			return ESQLitePreparedStatementExecuteRowResult::Error;
+		}) != INDEX_NONE;
+	}
+
 	/** Get the unique package_name_ids from package_transactions */
 	SQLITE_PREPARED_STATEMENT_COLUMNS_ONLY(FGetPackageNameIdsWithTransactions, "SELECT DISTINCT package_name_id FROM package_transactions;", SQLITE_PREPARED_STATEMENT_COLUMNS(int64));
 	FGetPackageNameIdsWithTransactions Statement_GetPackageNameIdsWithTransactions;
@@ -1847,16 +1864,24 @@ bool FConcertSyncSessionDatabase::GetPackageNamesWithLiveTransactions(TArray<FNa
 
 bool FConcertSyncSessionDatabase::EnumeratePackageNamesWithLiveTransactions(TFunctionRef<bool(FName)> InCallback) const
 {
-	return Statements->GetPackageNameIdsWithTransactions([this, &InCallback](const int64 InPackageNameId)
+	return Statements->GetPackageNameIdsMaxTransactionId([this, &InCallback](const int64 InPackageNameId, const int64 InMaxTransactionEventId)
 	{
-		FName PackageName;
-		if (GetPackageName(InPackageNameId, PackageName))
+		// Get the transaction id at last save, if the max transaction id for a package name id is greater than its transaction id at last save, it has live transactions
+		int64 HeadTransactionEventIdAtLastSave = 0;
+		Statements->GetPackageTransactionEventIdAtLastSave(InPackageNameId, HeadTransactionEventIdAtLastSave);
+		
+		if (InMaxTransactionEventId > HeadTransactionEventIdAtLastSave)
 		{
-			return InCallback(PackageName)
-				? ESQLitePreparedStatementExecuteRowResult::Continue
-				: ESQLitePreparedStatementExecuteRowResult::Stop;
+			FName PackageName;
+			if (GetPackageName(InPackageNameId, PackageName))
+			{
+				return InCallback(PackageName)
+					? ESQLitePreparedStatementExecuteRowResult::Continue
+					: ESQLitePreparedStatementExecuteRowResult::Stop;
+			}
+			return ESQLitePreparedStatementExecuteRowResult::Error;
 		}
-		return ESQLitePreparedStatementExecuteRowResult::Error;
+		return ESQLitePreparedStatementExecuteRowResult::Continue;
 	});
 }
 
