Edigrate cl 9643151:
Modified `GatherSessionChanges` to only report modified packages without saving live transactions to them deferring it to when files are actually persisted.
`GatherSessionChanges` and `PersistSessionChanges` now takes a list of package names instead of filenames
Gathering the list of changes is also now done entirely using the db instead of being split between the sandbox and the db.
Tracking of persist events are now done in the session db on the client side and is tracked through multiple connections with the same editor instance

diff --git a/Engine/Plugins/Developer/Concert/ConcertSync/ConcertSyncClient/Source/ConcertSyncClient/Private/ConcertClientPackageManager.cpp b/Engine/Plugins/Developer/Concert/ConcertSync/ConcertSyncClient/Source/ConcertSyncClient/Private/ConcertClientPackageManager.cpp
--- a/Engine/Plugins/Developer/Concert/ConcertSync/ConcertSyncClient/Source/ConcertSyncClient/Private/ConcertClientPackageManager.cpp
+++ b/Engine/Plugins/Developer/Concert/ConcertSync/ConcertSyncClient/Source/ConcertSyncClient/Private/ConcertClientPackageManager.cpp
@@ -212,30 +212,31 @@ void FConcertClientPackageManager::ApplyAllHeadPackageData()
 
 bool FConcertClientPackageManager::HasSessionChanges() const
 {
-#if WITH_EDITOR
-	return SandboxPlatformFile && SandboxPlatformFile->GatherSandboxChangedFilenames().Num() > 0;
-#else
-	return false;
-#endif	// WITH_EDITOR
-}
-
-TArray<FString> FConcertClientPackageManager::GatherSessionChanges()
-{
-#if WITH_EDITOR
-	if (SandboxPlatformFile)
+	bool bHasSessionChanges = false;
+	LiveSession->GetSessionDatabase().EnumeratePackageNamesWithHeadRevision([&bHasSessionChanges](const FName PackageName)
 	{
-		return SandboxPlatformFile->GatherSandboxChangedFilenames();
-	}
-#endif	// WITH_EDITOR
-	return TArray<FString>();
+		bHasSessionChanges = true;
+		return false; // Stop enumeration
+	}, /*IgnorePersisted*/true);
+	return bHasSessionChanges;
 }
 
-bool FConcertClientPackageManager::PersistSessionChanges(TArrayView<const FString> InFilesToPersist, ISourceControlProvider* SourceControlProvider, TArray<FText>* OutFailureReasons)
+bool FConcertClientPackageManager::PersistSessionChanges(TArrayView<const FName> InPackagesToPersist, ISourceControlProvider* SourceControlProvider, TArray<FText>* OutFailureReasons)
 {
 #if WITH_EDITOR
 	if (SandboxPlatformFile)
 	{
-		return SandboxPlatformFile->PersistSandbox(MoveTemp(InFilesToPersist), SourceControlProvider, OutFailureReasons);
+		// Transform all the package names into actual filenames
+		TArray<FString, TInlineAllocator<8>> FilesToPersist;
+		FString Filename;
+		for (const FName& PackageName : InPackagesToPersist)
+		{
+			if (FPackageName::DoesPackageExist(PackageName.ToString(), nullptr, &Filename))
+			{
+				FilesToPersist.Add(MoveTemp(Filename));
+			}
+		}
+		return SandboxPlatformFile->PersistSandbox(FilesToPersist, SourceControlProvider, OutFailureReasons);
 	}
 #endif
 	return false;
diff --git a/Engine/Plugins/Developer/Concert/ConcertSync/ConcertSyncClient/Source/ConcertSyncClient/Private/ConcertClientPackageManager.h b/Engine/Plugins/Developer/Concert/ConcertSync/ConcertSyncClient/Source/ConcertSyncClient/Private/ConcertClientPackageManager.h
--- a/Engine/Plugins/Developer/Concert/ConcertSync/ConcertSyncClient/Source/ConcertSyncClient/Private/ConcertClientPackageManager.h
+++ b/Engine/Plugins/Developer/Concert/ConcertSync/ConcertSyncClient/Source/ConcertSyncClient/Private/ConcertClientPackageManager.h
@@ -67,15 +67,9 @@ public:
 	bool HasSessionChanges() const;
 
 	/**
-	 * Gather package changes that happened in this session.
-	 * @return A list of package files that were modified during the session.
+	 * Persist the session changes from the package name list and prepare it for source control submission.
 	 */
-	TArray<FString> GatherSessionChanges();
-
-	/**
-	 * Persist the session changes from the file list and prepare it for source control submission.
-	 */
-	bool PersistSessionChanges(TArrayView<const FString> InFilesToPersist, ISourceControlProvider* SourceControlProvider, TArray<FText>* OutFailureReasons = nullptr);
+	bool PersistSessionChanges(TArrayView<const FName> InPackagesToPersist, ISourceControlProvider* SourceControlProvider, TArray<FText>* OutFailureReasons = nullptr);
 
 private:
 	/**
diff --git a/Engine/Plugins/Developer/Concert/ConcertSync/ConcertSyncClient/Source/ConcertSyncClient/Private/ConcertClientWorkspace.cpp b/Engine/Plugins/Developer/Concert/ConcertSync/ConcertSyncClient/Source/ConcertSyncClient/Private/ConcertClientWorkspace.cpp
--- a/Engine/Plugins/Developer/Concert/ConcertSync/ConcertSyncClient/Source/ConcertSyncClient/Private/ConcertClientWorkspace.cpp
+++ b/Engine/Plugins/Developer/Concert/ConcertSync/ConcertSyncClient/Source/ConcertSyncClient/Private/ConcertClientWorkspace.cpp
@@ -114,31 +114,52 @@ bool FConcertClientWorkspace::HasSessionChanges() const
 	return (TransactionManager && TransactionManager->HasSessionChanges()) || (PackageManager && PackageManager->HasSessionChanges());
 }
 
-TArray<FString> FConcertClientWorkspace::GatherSessionChanges()
+TArray<FName> FConcertClientWorkspace::GatherSessionChanges(bool IgnorePersisted)
 {
-	TArray<FString> SessionChanges;
-#if WITH_EDITOR
-	// Save live transactions to packages so we can properly report those changes.
-	SaveLiveTransactionsToPackages();
-#endif
-	// Persist the sandbox state over the real content directory
-	// This will also check things out from source control and make them ready to be submitted
-	if (PackageManager)
+	TSet<FName> SessionChangedPackageNames;
+
+	// Gather the packages with live transactions
+	LiveSession->GetSessionDatabase().EnumeratePackageNamesWithLiveTransactions([&SessionChangedPackageNames](FName PackageName)
 	{
-		SessionChanges = PackageManager->GatherSessionChanges();
-	}
-	return SessionChanges;
+		SessionChangedPackageNames.Add(PackageName);
+		return true;
+	});
+
+	// Gather the packages with a non persisted head revision events
+	LiveSession->GetSessionDatabase().EnumeratePackageNamesWithHeadRevision([&SessionChangedPackageNames](FName PackageName)
+	{
+		SessionChangedPackageNames.Add(PackageName);
+		return true;
+	}, IgnorePersisted);
+
+	return SessionChangedPackageNames.Array();
 }
 
-bool FConcertClientWorkspace::PersistSessionChanges(TArrayView<const FString> InFilesToPersist, ISourceControlProvider* SourceControlProvider, TArray<FText>* OutFailureReasons)
+bool FConcertClientWorkspace::PersistSessionChanges(TArrayView<const FName> InPackagesToPersist, ISourceControlProvider* SourceControlProvider, TArray<FText>* OutFailureReasons)
 {
+	bool bSuccess = false;
 #if WITH_EDITOR
 	if (PackageManager)
 	{
-		return PackageManager->PersistSessionChanges(InFilesToPersist, SourceControlProvider, OutFailureReasons);
+		for (const FName& PackageName : InPackagesToPersist)
+		{
+			SaveLiveTransactionsToPackage(PackageName);
+		}
+
+		bSuccess = PackageManager->PersistSessionChanges(InPackagesToPersist, SourceControlProvider, OutFailureReasons);
+
+		// if we successfully persisted the files, record persist events for them in the db
+		if (bSuccess)
+		{
+			int64 PersistEventId = 0;
+			for (const FName& PackageName : InPackagesToPersist)
+			{
+				LiveSession->GetSessionDatabase().AddPersistEventForHeadRevision(PackageName, PersistEventId);
+			}
+		}
 	}
 #endif
-	return false;
+	return bSuccess;
 }
 
 bool FConcertClientWorkspace::HasLiveTransactionSupport(UPackage* InPackage) const
@@ -516,10 +537,25 @@ void FConcertClientWorkspace::SaveLiveTransactionsToPackages()
 	// Save any packages that have live transactions
 	if (GEditor && TransactionManager)
 	{
-		// Ignore these package saves as the other clients should already be in-sync
-		IConcertClientPackageBridge::FScopedIgnoreLocalSave IgnorePackageSaveScope(*PackageBridge);
 		LiveSession->GetSessionDatabase().EnumeratePackageNamesWithLiveTransactions([this](const FName PackageName)
 		{
+			SaveLiveTransactionsToPackage(PackageName);
+			return true;
+		});
+	}
+}
+
+void FConcertClientWorkspace::SaveLiveTransactionsToPackage(const FName PackageName)
+{
+	if (GEditor && TransactionManager)
+	{
+		bool bHasLiveTransactions = false;
+		if (LiveSession->GetSessionDatabase().PackageHasLiveTransactions(PackageName, bHasLiveTransactions)
+			&& bHasLiveTransactions)
+		{
+			// Ignore these package saves as the other clients should already be in-sync
+			IConcertClientPackageBridge::FScopedIgnoreLocalSave IgnorePackageSaveScope(*PackageBridge);
+
 			const FString PackageNameStr = PackageName.ToString();
 			UPackage* Package = LoadPackage(nullptr, *PackageNameStr, LOAD_None);
 			if (Package)
@@ -549,8 +585,7 @@ void FConcertClientWorkspace::SaveLiveTransactionsToPackages()
 					UE_LOG(LogConcert, Warning, TEXT("Failed to save package '%s' when persiting sandbox state!"), *PackageNameStr);
 				}
 			}
-			return true;
-		});
+		}
 	}
 }
 
diff --git a/Engine/Plugins/Developer/Concert/ConcertSync/ConcertSyncClient/Source/ConcertSyncClient/Private/ConcertClientWorkspace.h b/Engine/Plugins/Developer/Concert/ConcertSync/ConcertSyncClient/Source/ConcertSyncClient/Private/ConcertClientWorkspace.h
--- a/Engine/Plugins/Developer/Concert/ConcertSync/ConcertSyncClient/Source/ConcertSyncClient/Private/ConcertClientWorkspace.h
+++ b/Engine/Plugins/Developer/Concert/ConcertSync/ConcertSyncClient/Source/ConcertSyncClient/Private/ConcertClientWorkspace.h
@@ -37,8 +37,8 @@ public:
 	virtual TFuture<FConcertResourceLockResponse> LockResources(TArray<FName> InResourceNames) override;
 	virtual TFuture<FConcertResourceLockResponse> UnlockResources(TArray<FName> InResourceNames) override;
 	virtual bool HasSessionChanges() const override;
-	virtual TArray<FString> GatherSessionChanges() override;
-	virtual bool PersistSessionChanges(TArrayView<const FString> InFilesToPersist, ISourceControlProvider* SourceControlProvider, TArray<FText>* OutFailureReasons = nullptr) override;
+	virtual TArray<FName> GatherSessionChanges(bool IgnorePersisted = true) override;
+	virtual bool PersistSessionChanges(TArrayView<const FName> InPackagesToPersist, ISourceControlProvider* SourceControlProvider, TArray<FText>* OutFailureReasons = nullptr) override;
 	virtual bool HasLiveTransactionSupport(UPackage* InPackage) const override;
 	virtual bool ShouldIgnorePackageDirtyEvent(class UPackage* InPackage) const override;
 	virtual bool FindTransactionEvent(const int64 TransactionEventId, FConcertSyncTransactionEvent& OutTransactionEvent, const bool bMetaDataOnly) const override;
@@ -74,9 +74,16 @@ private:
 	void HandleConnectionChanged(IConcertClientSession& InSession, EConcertConnectionStatus Status);
 
 #if WITH_EDITOR
-	/** */
+	/**
+	 * Save all live transactions to packages.
+	 */
 	void SaveLiveTransactionsToPackages();
 
+	/**
+	 * Save live transactions, if any, to the specified packages, adding a client side dummy package event to the db
+	 */
+	void SaveLiveTransactionsToPackage(const FName PackageName);
+
 	/** */
 	void HandleAssetLoaded(UObject* InAsset);
 
diff --git a/Engine/Plugins/Developer/Concert/ConcertSync/ConcertSyncClient/Source/ConcertSyncClient/Private/ConcertSyncClient.cpp b/Engine/Plugins/Developer/Concert/ConcertSync/ConcertSyncClient/Source/ConcertSyncClient/Private/ConcertSyncClient.cpp
--- a/Engine/Plugins/Developer/Concert/ConcertSync/ConcertSyncClient/Source/ConcertSyncClient/Private/ConcertSyncClient.cpp
+++ b/Engine/Plugins/Developer/Concert/ConcertSync/ConcertSyncClient/Source/ConcertSyncClient/Private/ConcertSyncClient.cpp
@@ -109,7 +109,7 @@ void FConcertSyncClient::PersistAllSessionChanges()
 #if WITH_EDITOR
 	if (Workspace)
 	{
-		TArray<FString> SessionChanges = Workspace->GatherSessionChanges();
+		TArray<FName> SessionChanges = Workspace->GatherSessionChanges();
 		Workspace->PersistSessionChanges(SessionChanges, SourceControlProxy.Get());
 	}
 #endif
diff --git a/Engine/Plugins/Developer/Concert/ConcertSync/ConcertSyncClient/Source/ConcertSyncClient/Public/IConcertClientWorkspace.h b/Engine/Plugins/Developer/Concert/ConcertSync/ConcertSyncClient/Source/ConcertSyncClient/Public/IConcertClientWorkspace.h
--- a/Engine/Plugins/Developer/Concert/ConcertSync/ConcertSyncClient/Source/ConcertSyncClient/Public/IConcertClientWorkspace.h
+++ b/Engine/Plugins/Developer/Concert/ConcertSync/ConcertSyncClient/Source/ConcertSyncClient/Public/IConcertClientWorkspace.h
@@ -98,12 +98,13 @@ public:
 
 	/**
 	 * Gather assets changes that happened on the workspace in this session.
-	 * @return a list of asset files that were modified during the session.
+	 * @param IgnorePersisted if true will not return packages which have already been persisted in their current state.
+	 * @return a list of package names that were modified during the session.
 	 */
-	virtual TArray<FString> GatherSessionChanges() = 0;
+	virtual TArray<FName> GatherSessionChanges(bool IgnorePersisted = true) = 0;
 
-	/** Persist the session changes from the file list and prepare it for source control submission */
-	virtual bool PersistSessionChanges(TArrayView<const FString> InFilesToPersist, ISourceControlProvider* SourceControlProvider, TArray<FText>* OutFailureReasonMap = nullptr) = 0;
+	/** Persist the session changes from the package list and prepare it for source control submission */
+	virtual bool PersistSessionChanges(TArrayView<const FName> InPackageToPersist, ISourceControlProvider* SourceControlProvider, TArray<FText>* OutFailureReasonMap = nullptr) = 0;
 
 	/**
 	 * Get Activities from the session.
diff --git a/Engine/Plugins/Developer/Concert/ConcertSync/ConcertSyncCore/Source/ConcertSyncCore/Private/ConcertSyncSessionDatabase.cpp b/Engine/Plugins/Developer/Concert/ConcertSync/ConcertSyncCore/Source/ConcertSyncCore/Private/ConcertSyncSessionDatabase.cpp
--- a/Engine/Plugins/Developer/Concert/ConcertSync/ConcertSyncCore/Source/ConcertSyncCore/Private/ConcertSyncSessionDatabase.cpp
+++ b/Engine/Plugins/Developer/Concert/ConcertSync/ConcertSyncCore/Source/ConcertSyncCore/Private/ConcertSyncSessionDatabase.cpp
@@ -301,6 +301,9 @@ public:
 		PREPARE_STATEMENT(Statement_GetPackageName);
 		PREPARE_STATEMENT(Statement_GetPackageNameId);
 
+		PREPARE_STATEMENT(Statement_GetPersistEventId);
+		PREPARE_STATEMENT(Statement_AddPersistEvent);
+
 		PREPARE_STATEMENT(Statement_SetEndpointData);
 		PREPARE_STATEMENT(Statement_GetEndpointDataForId);
 		PREPARE_STATEMENT(Statement_GetAllEndpointData);
@@ -325,6 +328,8 @@ public:
 		PREPARE_STATEMENT(Statement_GetPackageMaxEventId);
 		PREPARE_STATEMENT(Statement_GetPackageDataForRevision);
 		PREPARE_STATEMENT(Statement_GetPackageHeadEventId);
+		PREPARE_STATEMENT(Statement_GetPackageHeadEventIdAndTransactionIdAtSave);
+		PREPARE_STATEMENT(Statement_GetMaxPackageEventIdAndTransactionEventIdAtSavePerPackageNameId);
 		PREPARE_STATEMENT(Statement_GetPackageHeadRevison);
 		PREPARE_STATEMENT(Statement_GetPackageTransactionEventIdAtLastSave);
 
@@ -452,7 +457,7 @@ public:
 		return Statement_GetPackageName.BindAndExecuteSingle(InPackageNameId, OutPackageName);
 	}
 
-	/** Get an package_name_id from package_names for the given package_name */
+	/** Get a package_name_id from package_names for the given package_name */
 	SQLITE_PREPARED_STATEMENT(FGetPackageNameId, "SELECT package_name_id FROM package_names WHERE package_name = ?1;", SQLITE_PREPARED_STATEMENT_COLUMNS(int64), SQLITE_PREPARED_STATEMENT_BINDINGS(FName));
 	FGetPackageNameId Statement_GetPackageNameId;
 	bool GetPackageNameId(const FName InPackageName, int64& OutPackageNameId)
@@ -460,6 +465,31 @@ public:
 		return Statement_GetPackageNameId.BindAndExecuteSingle(InPackageName, OutPackageNameId);
 	}
 
+	/**
+	 * Statements working on persist_events
+	 */
+
+	/** Get a persist_event_id and transaction_event_id_at_persist from persist_events for the given package_event_id */
+	SQLITE_PREPARED_STATEMENT(FGetPersistEventId, "SELECT persist_event_id, transaction_event_id_at_persist FROM persist_events WHERE package_event_id = ?1;", SQLITE_PREPARED_STATEMENT_COLUMNS(int64, int64), SQLITE_PREPARED_STATEMENT_BINDINGS(int64));
+	FGetPersistEventId Statement_GetPersistEventId;
+	bool GetPersistEventId(int64 InPackageEventId, int64& OutPersistEventId, int64& OutTransactionEventIdAtPersist)
+	{
+		return Statement_GetPersistEventId.BindAndExecuteSingle(InPackageEventId, OutPersistEventId, OutTransactionEventIdAtPersist);
+	}
+
+	/** Add a new package_event_id to persist_events and get its persist_event_id. */
+	SQLITE_PREPARED_STATEMENT_BINDINGS_ONLY(FAddPersistEvent, "INSERT INTO persist_events(package_event_id, transaction_event_id_at_persist) VALUES(?1, ?2);", SQLITE_PREPARED_STATEMENT_BINDINGS(int64, int64));
+	FAddPersistEvent Statement_AddPersistEvent;
+	bool AddPersistEvent(int64 InPackageEventId, int64 InTransactionEventIdAtPersist, int64& OutPersistEventId)
+	{
+		if (Statement_AddPersistEvent.BindAndExecute(InPackageEventId, InTransactionEventIdAtPersist))
+		{
+			OutPersistEventId = Database.GetLastInsertRowId();
+			return true;
+		}
+		return false;
+	}
+
 	/**
 	 * Statements working on endpoints
 	 */
@@ -703,6 +733,30 @@ public:
 		return Statement_GetPackageHeadEventId.BindAndExecuteSingle(InPackageNameId, OutPackageEventId);
 	}
 
+	/** Get the largest package_event_id and its transaction_event_id_at_save currently in package_events for the given package_name_id */
+	SQLITE_PREPARED_STATEMENT(FGetPackageHeadEventIdAndTransactionIdAtSave, "SELECT MAX(package_event_id), transaction_event_id_at_save FROM package_events WHERE package_name_id = ?1;", SQLITE_PREPARED_STATEMENT_COLUMNS(int64, int64), SQLITE_PREPARED_STATEMENT_BINDINGS(int64));
+	FGetPackageHeadEventIdAndTransactionIdAtSave Statement_GetPackageHeadEventIdAndTransactionIdAtSave;
+	bool GetPackageHeadEventIdAndTransactionIdAtSave(const int64 InPackageNameId, int64& OutPackageEventId, int64& OutTransactionEventIdAtSave)
+	{
+		return Statement_GetPackageHeadEventIdAndTransactionIdAtSave.BindAndExecuteSingle(InPackageNameId, OutPackageEventId, OutTransactionEventIdAtSave);
+	}
+
+	/** Get the largest package_event_id along its transaction_event_id_at_save currently in package_events for each distinct package_name_id */
+	SQLITE_PREPARED_STATEMENT_COLUMNS_ONLY(FGetMaxPackageEventIdAndTransactionEventIdAtSavePerPackageNameId, "SELECT package_name_id, MAX(package_event_id), transaction_event_id_at_save FROM package_events GROUP BY package_name_id;", SQLITE_PREPARED_STATEMENT_COLUMNS(int64, int64, int64));
+	FGetMaxPackageEventIdAndTransactionEventIdAtSavePerPackageNameId Statement_GetMaxPackageEventIdAndTransactionEventIdAtSavePerPackageNameId;
+	bool GetMaxPackageEventIdAndTransactionEventIdAtSavePerPackageNameId(TFunctionRef<ESQLitePreparedStatementExecuteRowResult(int64, int64, int64)> InCallback)
+	{
+		return Statement_GetMaxPackageEventIdAndTransactionEventIdAtSavePerPackageNameId.Execute([&InCallback](const FGetMaxPackageEventIdAndTransactionEventIdAtSavePerPackageNameId& InStatement)
+		{
+			int64 PackageNameId = 0, MaxPackageEventId = 0, TransactionEventIdAtSave = 0;
+			if (InStatement.GetColumnValues(PackageNameId, MaxPackageEventId, TransactionEventIdAtSave))
+			{
+				return InCallback(PackageNameId, MaxPackageEventId, TransactionEventIdAtSave);
+			}
+			return ESQLitePreparedStatementExecuteRowResult::Error;
+		}) != INDEX_NONE;
+	}
+
 	/** Get the largest package_revision currently in package_events for the given package_name_id */
 	SQLITE_PREPARED_STATEMENT(FGetPackageHeadRevison, "SELECT MAX(package_revision) FROM package_events WHERE package_name_id = ?1;", SQLITE_PREPARED_STATEMENT_COLUMNS(int64), SQLITE_PREPARED_STATEMENT_BINDINGS(int64));
 	FGetPackageHeadRevison Statement_GetPackageHeadRevison;
@@ -1205,6 +1259,7 @@ bool FConcertSyncSessionDatabase::Open(const FString& InSessionPath, const ESQLi
 	CREATE_TABLE("lock_events", "lock_event_id INTEGER PRIMARY KEY, lock_event_type INTEGER NOT NULL");
 	CREATE_TABLE("transaction_events", "transaction_event_id INTEGER PRIMARY KEY, data_filename TEXT NOT NULL");
 	CREATE_TABLE("package_events", "package_event_id INTEGER PRIMARY KEY, package_name_id INTEGER NOT NULL, package_revision INTEGER NOT NULL, package_info_size_bytes INTEGER NOT NULL, package_info_data BLOB, transaction_event_id_at_save INTEGER NOT NULL, data_filename TEXT NOT NULL, FOREIGN KEY(package_name_id) REFERENCES package_names(package_name_id)");
+	CREATE_TABLE("persist_events", "persist_event_id INTEGER PRIMARY KEY, package_event_id INTEGER NOT NULL, transaction_event_id_at_persist INTEGER NOT NULL, FOREIGN KEY(package_event_id) REFERENCES package_events(package_event_id)");
 	CREATE_TABLE("activities", "activity_id INTEGER PRIMARY KEY, endpoint_id BLOB NOT NULL, event_time INTEGER NOT NULL, event_type INTEGER NOT NULL, event_id INTEGER NOT NULL, event_summary_type TEXT NOT NULL, event_summary_size_bytes INTEGER NOT NULL, event_summary_data BLOB, FOREIGN KEY(endpoint_id) REFERENCES endpoints(endpoint_id)");
 	CREATE_TABLE("ignored_activities", "activity_id INTEGER NOT NULL, FOREIGN KEY(activity_id) REFERENCES activities(activity_id)");
 	CREATE_TABLE("resource_locks", "object_name_id INTEGER NOT NULL, lock_event_id INTEGER NOT NULL, FOREIGN KEY(object_name_id) REFERENCES object_names(object_name_id), FOREIGN KEY(lock_event_id) REFERENCES lock_events(lock_event_id)");
@@ -1228,6 +1283,7 @@ bool FConcertSyncSessionDatabase::Open(const FString& InSessionPath, const ESQLi
 	CREATE_UNIQUE_INDEX("idx_object_path_names_in_object_names", "object_names", "object_path_name");
 	CREATE_UNIQUE_INDEX("idx_package_names_in_package_names", "package_names", "package_name");
 	CREATE_INDEX("idx_package_name_ids_in_package_events", "package_events", "package_name_id");
+	CREATE_INDEX("idx_package_event_ids_in_persist_events", "persist_events", "package_event_id");
 	CREATE_INDEX("idx_event_ids_in_activities", "activities", "event_id");
 	CREATE_UNIQUE_INDEX("idx_activity_ids_in_ignored_activities", "ignored_activities", "activity_id");
 	CREATE_INDEX("idx_object_name_ids_in_resource_locks", "resource_locks", "object_name_id");
@@ -1841,6 +1897,16 @@ bool FConcertSyncSessionDatabase::GetLiveTransactionEventIdsForPackage(const FNa
 	});
 }
 
+bool FConcertSyncSessionDatabase::PackageHasLiveTransactions(const FName InPackageName, bool& OutHasLiveTransaction) const
+{
+	OutHasLiveTransaction = false;
+	return EnumerateLiveTransactionEventIdsForPackage(InPackageName, [&OutHasLiveTransaction](const int64)
+	{
+		OutHasLiveTransaction = true;
+		return false;
+	});
+}
+
 bool FConcertSyncSessionDatabase::EnumerateLiveTransactionEventIdsForPackage(const FName InPackageName, TFunctionRef<bool(int64)> InCallback) const
 {
 	int64 PackageNameId = 0;
@@ -2036,6 +2102,56 @@ bool FConcertSyncSessionDatabase::GetPackageEvent(const int64 InPackageEventId,
 	return false;
 }
 
+bool FConcertSyncSessionDatabase::GetPackageNamesWithHeadRevision(TArray<FName>& OutPackageNames, bool IgnorePersisted) const
+{
+	OutPackageNames.Reset();
+	return EnumeratePackageNamesWithHeadRevision([&OutPackageNames](const FName InPackageName)
+	{
+		OutPackageNames.Add(InPackageName);
+		return true;
+	}, IgnorePersisted);
+}
+
+bool FConcertSyncSessionDatabase::EnumeratePackageNamesWithHeadRevision(TFunctionRef<bool(FName)> InCallback, bool IgnorePersisted) const
+{
+	// if we ignore packages with persist event we need to compare head revision against entry in the persist table.
+	if (IgnorePersisted)
+	{
+		return Statements->GetMaxPackageEventIdAndTransactionEventIdAtSavePerPackageNameId([this, &InCallback](int64 InPackageNameId, int64 InMaxPackageEventId, int64 TransactionEventIdAtSave)
+		{
+			// We enumerate the packages if there isn't an entry in persist events with this MaxPackageEventId or 
+			// if the TransactionEventIdAtPersist and TransactionEventIdAtSave doesn't match (in case a dummy event got squashed)
+			int64 PersistEventId = 0, TransactionEventIdAtPersist = 0;
+			if (!Statements->GetPersistEventId(InMaxPackageEventId, PersistEventId, TransactionEventIdAtPersist)
+				|| TransactionEventIdAtPersist != TransactionEventIdAtSave)
+			{
+				FName PackageName;
+				if (GetPackageName(InPackageNameId, PackageName))
+				{
+					return InCallback(PackageName)
+						? ESQLitePreparedStatementExecuteRowResult::Continue
+						: ESQLitePreparedStatementExecuteRowResult::Stop;
+				}
+				return ESQLitePreparedStatementExecuteRowResult::Error;
+			}
+			return ESQLitePreparedStatementExecuteRowResult::Continue;
+		});
+	}
+	
+	// otherwise we can just gather distinct packages in the package events table
+	return Statements->GetUniquePackageNameIdsForPackageEvents([this, &InCallback](int64 InPackageNameId)
+	{
+		FName PackageName;
+		if (GetPackageName(InPackageNameId, PackageName))
+		{
+			return InCallback(PackageName)
+				? ESQLitePreparedStatementExecuteRowResult::Continue
+				: ESQLitePreparedStatementExecuteRowResult::Stop;
+		}
+		return ESQLitePreparedStatementExecuteRowResult::Error;
+	});
+}
+
 bool FConcertSyncSessionDatabase::EnumerateHeadRevisionPackageData(TFunctionRef<bool(FConcertPackage&&)> InCallback, const bool InMetaDataOnly) const
 {
 	return Statements->GetUniquePackageNameIdsForPackageEvents([this, &InCallback, InMetaDataOnly](int64 InPackageNameId)
@@ -2065,6 +2181,18 @@ bool FConcertSyncSessionDatabase::GetPackageMaxEventId(int64& OutPackageEventId)
 	return Statements->GetPackageMaxEventId(OutPackageEventId);
 }
 
+bool FConcertSyncSessionDatabase::AddPersistEventForHeadRevision(FName InPackageName, int64& OutPersistEventId)
+{
+	int64 PackageNameId = 0, HeadPackageEventId = 0, TransactionEventIdAtSave = 0;
+	if (GetPackageNameId(InPackageName, PackageNameId) 
+		&& Statements->GetPackageHeadEventIdAndTransactionIdAtSave(PackageNameId, HeadPackageEventId, TransactionEventIdAtSave)
+		&& HeadPackageEventId > 0)
+	{
+		Statements->AddPersistEvent(HeadPackageEventId, TransactionEventIdAtSave, OutPersistEventId);
+	}
+	return false;
+}
+
 bool FConcertSyncSessionDatabase::GetPackageDataForRevision(const FName InPackageName, FConcertPackage& OutPackage, const int64* InPackageRevision) const
 {
 	return GetPackageDataForRevision(InPackageName, OutPackage.Info, &OutPackage.PackageData, InPackageRevision);
diff --git a/Engine/Plugins/Developer/Concert/ConcertSync/ConcertSyncCore/Source/ConcertSyncCore/Public/ConcertSyncSessionDatabase.h b/Engine/Plugins/Developer/Concert/ConcertSync/ConcertSyncCore/Source/ConcertSyncCore/Public/ConcertSyncSessionDatabase.h
--- a/Engine/Plugins/Developer/Concert/ConcertSync/ConcertSyncCore/Source/ConcertSyncCore/Public/ConcertSyncSessionDatabase.h
+++ b/Engine/Plugins/Developer/Concert/ConcertSync/ConcertSyncCore/Source/ConcertSyncCore/Public/ConcertSyncSessionDatabase.h
@@ -480,6 +480,16 @@ public:
 	 */
 	bool GetLiveTransactionEventIdsForPackage(const FName InPackageName, TArray<int64>& OutTransactionEventIds) const;
 
+	/**
+	 * Get if a Package has any live transactions
+	 *
+	 * @param InPackageName				The name of the package to check if it has live transactions.
+	 * @param OutHasLiveTransaction		Bool to populate with the result.
+	 *
+	 * @return True if query was resolved correctly, false otherwise.
+	 */
+	bool PackageHasLiveTransactions(const FName InPackageName, bool& OutHasLiveTransaction) const;
+
 	/**
 	 * Enumerate the IDs of any live transaction events for the given package name.
 	 *
@@ -530,6 +540,26 @@ public:
 	 */
 	bool GetPackageEvent(const int64 InPackageEventId, FConcertSyncPackageEvent& OutPackageEvent, const bool InMetaDataOnly = false) const;
 
+	/**
+	 * Get package names for packages with a head revision (at least one package event)
+	 *
+	 * @param OutPackageNames			The array of names to populate with the result.
+	 * @param IgnorePersisted			Will ignore packages which head revision have been persisted.
+	 *
+	 * @return True if the package names were resolved, false otherwise.
+	 */
+	bool GetPackageNamesWithHeadRevision(TArray<FName>& OutPackageNames, bool IgnorePersisted) const;
+
+	/**
+	 * Enumerate package names for packages with a head revision (at least one package event)
+	 *
+	 * @param InCallback				Callback invoked for each package name; return true to continue enumeration, or false to stop.
+	 * @param IgnorePersisted			Will skip enumeration of packages which head revision have been persisted.
+	 *
+	 * @return True if the package data was enumerated without error, false otherwise.
+	 */
+	bool EnumeratePackageNamesWithHeadRevision(TFunctionRef<bool(FName)> InCallback, bool IgnorePersisted) const;
+
 	/**
 	 * Enumerate the head revision package data for all packages in this database.
 	 *
@@ -583,6 +613,16 @@ public:
 	 */
 	bool IsHeadRevisionPackageEvent(const int64 InPackageEventId, bool& OutIsHeadRevision) const;
 
+	/**
+	 * Add a package event ID for the head revision to the persist events in this database, if not already existing.
+	 *
+	 * @param PackageName				The package name to add an event for.
+	 * @param OutPersistEventId			Populated with the ID of the persist event in the database.
+	 *
+	 * @return True if the persist event was added, false otherwise.
+	 */
+	bool AddPersistEventForHeadRevision(FName InPackageName, int64& OutPersistEventId);
+
 	/**
 	 * Update the specified transaction event.
 	 * @note The function is meant to update a transaction event that was partially synced to store the corresponding transaction data.
diff --git a/Engine/Plugins/Developer/Concert/MultiUserClient/Source/MultiUserClient/Private/ConcertWorkspaceUI.cpp b/Engine/Plugins/Developer/Concert/MultiUserClient/Source/MultiUserClient/Private/ConcertWorkspaceUI.cpp
--- a/Engine/Plugins/Developer/Concert/MultiUserClient/Source/MultiUserClient/Private/ConcertWorkspaceUI.cpp
+++ b/Engine/Plugins/Developer/Concert/MultiUserClient/Source/MultiUserClient/Private/ConcertWorkspaceUI.cpp
@@ -589,12 +589,30 @@ bool FConcertWorkspaceUI::HasSessionChanges() const
 
 bool FConcertWorkspaceUI::PromptPersistSessionChanges()
 {
-	TArray<FSourceControlStateRef> States;
 	TSharedPtr<IConcertClientWorkspace> ClientWorkspacePin = ClientWorkspace.Pin();
+	TArray<TSharedPtr<FConcertPersistItem>> PersistItems;
 	if (ClientWorkspacePin.IsValid())
 	{
-		// Get file status of packages and config
-		ISourceControlModule::Get().GetProvider().GetState(ClientWorkspacePin->GatherSessionChanges(), States, EStateCacheUsage::ForceUpdate);
+		// Get source control status of packages 
+		TArray<FName> PackageNames;
+		TArray<FString> PackageFilenames;
+		TArray<FSourceControlStateRef> States;
+		PackageNames = ClientWorkspacePin->GatherSessionChanges();
+		FString Filename;
+		for (const FName& PackageName : PackageNames)
+		{
+			if (FPackageName::DoesPackageExist(PackageName.ToString(), nullptr, &Filename))
+			{
+				PackageFilenames.Add(FPaths::ConvertRelativePathToFull(MoveTemp(Filename)));
+			}
+		}
+		ISourceControlModule::Get().GetProvider().GetState(PackageFilenames, States, EStateCacheUsage::ForceUpdate);
+
+		// Create the list of persist items from the package names and source control state
+		for (int32 Index = 0; Index < PackageNames.Num(); ++Index)
+		{
+			PersistItems.Add(MakeShared<FConcertPersistItem>(PackageNames[Index], States[Index]));
+		}
 	}
 	
 	TSharedRef<SWindow> NewWindow = SNew(SWindow)
@@ -607,7 +625,7 @@ bool FConcertWorkspaceUI::PromptPersistSessionChanges()
 	TSharedRef<SConcertSandboxPersistWidget> PersistWidget =
 		SNew(SConcertSandboxPersistWidget)
 		.ParentWindow(NewWindow)
-		.Items(States);
+		.Items(PersistItems);
 
 	NewWindow->SetContent(
 		PersistWidget
@@ -630,7 +648,7 @@ bool FConcertWorkspaceUI::PromptPersistSessionChanges()
 	FText NotificationSub;
 
 	TArray<FText> PersistFailures;
-	bool bSuccess = ClientWorkspacePin->PersistSessionChanges(PersistCmd.FilesToPersist, &ISourceControlModule::Get().GetProvider(), &PersistFailures);
+	bool bSuccess = ClientWorkspacePin->PersistSessionChanges(PersistCmd.PackagesToPersist, &ISourceControlModule::Get().GetProvider(), &PersistFailures);
 	if (bSuccess)
 	{
 		bSuccess = SubmitChangelist(PersistCmd, NotificationSub);
diff --git a/Engine/Plugins/Developer/Concert/MultiUserClient/Source/MultiUserClient/Private/Widgets/SConcertSandboxPersistWidget.cpp b/Engine/Plugins/Developer/Concert/MultiUserClient/Source/MultiUserClient/Private/Widgets/SConcertSandboxPersistWidget.cpp
--- a/Engine/Plugins/Developer/Concert/MultiUserClient/Source/MultiUserClient/Private/Widgets/SConcertSandboxPersistWidget.cpp
+++ b/Engine/Plugins/Developer/Concert/MultiUserClient/Source/MultiUserClient/Private/Widgets/SConcertSandboxPersistWidget.cpp
@@ -30,52 +30,6 @@ namespace SConcertSandboxPersistWidgetDefs
 	const float IconColumnWidth = 21.0f;
 };
 
-/**
- * Persist Widget Row Item View Model
- */
-class FConcertPersistItem : public TSharedFromThis<FConcertPersistItem>
-{
-public:
-	/** Constructor */
-	explicit FConcertPersistItem(const FSourceControlStateRef& InItem)
-		: Item(InItem)
-	{
-		CheckBoxState = ECheckBoxState::Checked;
-		DisplayName = FText::FromString(Item->GetFilename());
-	}
-
-
-	/** Returns the full path of the item in source control */
-	FString GetFilename() const { return Item->GetFilename(); }
-
-	/** Returns the name of the item as displayed in the widget */
-	FText GetDisplayName() const { return DisplayName; }
-
-	/** Returns the name of the icon to be used in the list item widget */
-	FName GetIconName() const { return Item->GetSmallIconName(); }
-
-	/** Returns the tooltip text for the icon */
-	FText GetIconTooltip() const { return Item->GetDisplayTooltip(); }
-
-	/** Returns the checkbox state of this item */
-	ECheckBoxState GetCheckBoxState() const { return CheckBoxState; }
-
-	/** Sets the checkbox state of this item */
-	void SetCheckBoxState(ECheckBoxState NewState) { CheckBoxState = NewState; }
-
-	/** true if the item is enabled in the list */
-	bool IsEnabled() const { return true; }
-private:
-	/** Reference to the source control state */
-	FSourceControlStateRef Item;
-
-	/** Checkbox state */
-	ECheckBoxState CheckBoxState;
-
-	/** Cached name to display in the listview */
-	FText DisplayName;
-};
-
 /** Persist Widget Row */
 class SConcertSandboxPersistListRow : public SMultiColumnTableRow<TSharedPtr<FConcertPersistItem>>
 {
@@ -143,10 +97,7 @@ void SConcertSandboxPersistWidget::Construct(const FArguments& InArgs)
 {
 	bDialogConfirmed = false;
 	ParentWindow = InArgs._ParentWindow;
-	Algo::Transform(InArgs._Items, ListViewItems, [](const FSourceControlStateRef& SCCState)
-	{
-		return MakeShared<FConcertPersistItem>(SCCState);
-	});
+	ListViewItems = InArgs._Items;
 
 	SortByColumn = SConcertSandboxPersistWidgetDefs::ColumnID_FileLabel;
 	SortMode = EColumnSortMode::Ascending;
@@ -307,16 +258,17 @@ FConcertPersistCommand SConcertSandboxPersistWidget::GetPersistCommand() const
 	FConcertPersistCommand Cmd;
 	Cmd.bShouldSubmit = !IsSubmitDescriptionReadOnly();
 	Cmd.ChangelistDescription = SubmitDescriptionTextCtrl->GetText();
+	Cmd.PackagesToPersist.Reserve(ListViewItems.Num());
 	Cmd.FilesToPersist.Reserve(ListViewItems.Num());
-	Algo::TransformIf(ListViewItems, Cmd.FilesToPersist,
-		[](const TSharedPtr<FConcertPersistItem>& Item)
-		{
-			return Item->GetCheckBoxState() == ECheckBoxState::Checked;
-		},
-		[](const TSharedPtr<FConcertPersistItem>& Item)
+
+	for (const TSharedPtr<FConcertPersistItem>& Item : ListViewItems)
+	{
+		if (Item->GetCheckBoxState() == ECheckBoxState::Checked)
 		{
-			return Item->GetFilename();
-		});
+			Cmd.PackagesToPersist.Add(Item->GetPackageName());
+			Cmd.FilesToPersist.Add(Item->GetFilename());
+		}
+	}
 	return Cmd;
 }
 
diff --git a/Engine/Plugins/Developer/Concert/MultiUserClient/Source/MultiUserClient/Private/Widgets/SConcertSandboxPersistWidget.h b/Engine/Plugins/Developer/Concert/MultiUserClient/Source/MultiUserClient/Private/Widgets/SConcertSandboxPersistWidget.h
--- a/Engine/Plugins/Developer/Concert/MultiUserClient/Source/MultiUserClient/Private/Widgets/SConcertSandboxPersistWidget.h
+++ b/Engine/Plugins/Developer/Concert/MultiUserClient/Source/MultiUserClient/Private/Widgets/SConcertSandboxPersistWidget.h
@@ -17,11 +17,64 @@ class FConcertPersistItem;
 
 struct FConcertPersistCommand
 {
+	TArray<FName> PackagesToPersist;
 	TArray<FString> FilesToPersist;
 	FText ChangelistDescription;
 	bool bShouldSubmit;
 };
 
+/**
+ * Persist Widget Row Item View Model
+ */
+class FConcertPersistItem : public TSharedFromThis<FConcertPersistItem>
+{
+public:
+	/** Constructor */
+	explicit FConcertPersistItem(FName InPackageName, const FSourceControlStateRef& InItem)
+		: Item(InItem)
+		, PackageName(MoveTemp(InPackageName))
+	{
+		CheckBoxState = ECheckBoxState::Checked;
+		DisplayName = FText::FromString(Item->GetFilename());
+	}
+
+	/** Returns the package name of this persist item */
+	FName GetPackageName() const { return PackageName; }
+
+	/** Returns the full path of the item in source control */
+	FString GetFilename() const { return Item->GetFilename(); }
+
+	/** Returns the name of the item as displayed in the widget */
+	FText GetDisplayName() const { return DisplayName; }
+
+	/** Returns the name of the icon to be used in the list item widget */
+	FName GetIconName() const { return Item->GetSmallIconName(); }
+
+	/** Returns the tooltip text for the icon */
+	FText GetIconTooltip() const { return Item->GetDisplayTooltip(); }
+
+	/** Returns the checkbox state of this item */
+	ECheckBoxState GetCheckBoxState() const { return CheckBoxState; }
+
+	/** Sets the checkbox state of this item */
+	void SetCheckBoxState(ECheckBoxState NewState) { CheckBoxState = NewState; }
+
+	/** true if the item is enabled in the list */
+	bool IsEnabled() const { return true; }
+private:
+	/** Reference to the source control state */
+	FSourceControlStateRef Item;
+
+	/** Checkbox state */
+	ECheckBoxState CheckBoxState;
+
+	/** Cached name to display in the listview */
+	FText DisplayName;
+
+	/** The associated package name for this persist item. */
+	FName PackageName;
+};
+
 /**
  * Persist sandbox window
  */
@@ -32,7 +85,7 @@ public:
 		/** The parent window this widget is hosted in. */
 		SLATE_ARGUMENT(TSharedPtr<SWindow>, ParentWindow)
 		/** The file list to display. */
-		SLATE_ARGUMENT(TArray<FSourceControlStateRef>, Items)
+		SLATE_ARGUMENT(TArray<TSharedPtr<FConcertPersistItem>>, Items)
 	SLATE_END_ARGS()
 
 	void Construct(const FArguments& InArgs);
