Move ConcertSyncCore plugin Engine.ini plugin override to BaseConcertSyncCore.ini base settings to prevent the plugin from locking out saving of Engine.ini settings.

diff --git a/Engine/Plugins/Developer/Concert/ConcertSync/ConcertSyncCore/Config/BaseConcertSyncCore.ini b/Engine/Plugins/Developer/Concert/ConcertSync/ConcertSyncCore/Config/BaseConcertSyncCore.ini
new file mode 100644
--- /dev/null
+++ b/Engine/Plugins/Developer/Concert/ConcertSync/ConcertSyncCore/Config/BaseConcertSyncCore.ini
@@ -0,0 +1,14 @@
+[/Script/ConcertSyncCore.ConcertSyncConfig]
++IncludeObjectClassFilters=(ObjectOuterClass=None,ObjectClasses=(Class'/Script/Engine.World',Class'/Script/Engine.BookmarkBase',Class'/Script/Engine.MaterialParameterCollection',Class'/Script/MovieScene.MovieSceneSignedObject',Class'/Script/MovieScene.MovieSceneFolder'))
++IncludeObjectClassFilters=(ObjectOuterClass=Class'/Script/Engine.World',ObjectClasses=(Class'/Script/Engine.Layer',Class'/Script/Engine.Level',Class'/Script/Engine.LevelStreaming',Class'/Script/Engine.Actor',Class'/Script/Engine.ActorComponent',Class'/Script/Engine.Polys',Class'/Script/Engine.Model',Class'/Script/Engine.BrushBuilder',Class'/Script/Foliage.FoliageType'))
++IncludeObjectClassFilters=(ObjectOuterClass=Class'/Script/MovieScene.MovieSceneSignedObject',ObjectClasses=(Class'/Script/Engine.Actor',Class'/Script/Engine.ActorComponent'))
++ExcludeTransactionClassFilters=(ObjectOuterClass=None,ObjectClasses=(Class'/Script/Landscape.LandscapeProxy',Class'/Script/Landscape.LandscapeHeightfieldCollisionComponent',Class'/Script/Landscape.LandscapeComponent'))
++AllowedTransientProperties=Property'/Script/Engine.Actor:bHiddenEdTemporary'
++AllowedTransientProperties=Property'/Script/Engine.Actor:bHiddenEdLevel'
++AllowedTransientProperties=Property'/Script/Engine.Level:bIsVisible'
+
+; If you wish to use Composure in Multi-user, you may add that line under [/Script/ConcertSyncCore.ConcertSyncConfig] in your project engine config settings.
+;+IncludeObjectClassFilters=(ObjectOuterClass=None,ObjectClasses=(Class'/Script/Composure.CompositingElementPass',Class'/Script/Composure.ComposurePostProcessPassPolicy'))
+
+; If you wish to use LiveLink in Multi-user, you may add that line under [/Script/ConcertSyncCore.ConcertSyncConfig] in your project engine config settings.
+;+IncludeObjectClassFilters=(ObjectOuterClass=Class'/Script/LiveLinkComponents.LiveLinkComponentController',ObjectClasses=(Class'/Script/LiveLinkComponents.LiveLinkControllerBase'))
diff --git a/Engine/Plugins/Developer/Concert/ConcertSync/ConcertSyncCore/Config/Engine.ini b/Engine/Plugins/Developer/Concert/ConcertSync/ConcertSyncCore/Config/Engine.ini
deleted file mode 100644
diff --git a/Engine/Plugins/Developer/Concert/ConcertSync/ConcertSyncCore/Source/ConcertSyncCore/Public/ConcertSyncSettings.h b/Engine/Plugins/Developer/Concert/ConcertSync/ConcertSyncCore/Source/ConcertSyncCore/Public/ConcertSyncSettings.h
--- a/Engine/Plugins/Developer/Concert/ConcertSync/ConcertSyncCore/Source/ConcertSyncCore/Public/ConcertSyncSettings.h
+++ b/Engine/Plugins/Developer/Concert/ConcertSync/ConcertSyncCore/Source/ConcertSyncCore/Public/ConcertSyncSettings.h
@@ -24,7 +24,7 @@ struct CONCERTSYNCCORE_API FTransactionClassFilter
 	TArray<FSoftClassPath> ObjectClasses;
 };
 
-UCLASS(config=Engine)
+UCLASS(config=ConcertSyncCore)
 class CONCERTSYNCCORE_API UConcertSyncConfig : public UObject
 {
 	GENERATED_BODY()
diff --git a/Engine/Source/Runtime/Projects/Private/PluginManager.cpp b/Engine/Source/Runtime/Projects/Private/PluginManager.cpp
--- a/Engine/Source/Runtime/Projects/Private/PluginManager.cpp
+++ b/Engine/Source/Runtime/Projects/Private/PluginManager.cpp
@@ -708,6 +708,7 @@ bool FPluginManager::ConfigureEnabledPlugins()
 					GConfig->Remove(PluginConfigFilename);
 				}
 
+				//@note: This function is called too early for `GIsEditor` to be true and hence not go through this scope
 				if (!GIsEditor)
 				{
 					// override config cache entries with plugin configs (Engine.ini, Game.ini, etc in <PluginDir>\Config\)
