Added the ability to flag Concert activity that should be tracked, but ignored when restoring activity

diff --git a/Engine/Plugins/Developer/Concert/ConcertMain/Source/Concert/Public/ConcertMessageData.h b/Engine/Plugins/Developer/Concert/ConcertMain/Source/Concert/Public/ConcertMessageData.h
--- a/Engine/Plugins/Developer/Concert/ConcertMain/Source/Concert/Public/ConcertMessageData.h
+++ b/Engine/Plugins/Developer/Concert/ConcertMain/Source/Concert/Public/ConcertMessageData.h
@@ -199,7 +199,7 @@ struct FConcertSessionFilter
 
 	/**
 	 * Return true if the given activity ID passes the ID tests of this filter.
-	 * @note This function only tests the ID conditions, not any data specific checks like bOnlyLiveData.
+	 * @note This function only tests the ID conditions, not any data specific checks like bOnlyLiveData and bIncludeIgnoredActivities.
 	 */
 	CONCERT_API bool ActivityIdPassesFilter(const int64 InActivityId) const;
 
@@ -226,6 +226,10 @@ struct FConcertSessionFilter
 	/** True to export the activity summaries without the package/transaction data to look at the log rather than replaying the activities. */
 	UPROPERTY()
 	bool bMetaDataOnly = false;
+
+	/** True to include ignored activities */
+	UPROPERTY()
+	bool bIncludeIgnoredActivities = true;
 };
 
 USTRUCT()
diff --git a/Engine/Plugins/Developer/Concert/ConcertSync/ConcertSyncClient/Source/ConcertSyncClient/Private/ConcertClientWorkspace.cpp b/Engine/Plugins/Developer/Concert/ConcertSync/ConcertSyncClient/Source/ConcertSyncClient/Private/ConcertClientWorkspace.cpp
--- a/Engine/Plugins/Developer/Concert/ConcertSync/ConcertSyncClient/Source/ConcertSyncClient/Private/ConcertClientWorkspace.cpp
+++ b/Engine/Plugins/Developer/Concert/ConcertSync/ConcertSyncClient/Source/ConcertSyncClient/Private/ConcertClientWorkspace.cpp
@@ -707,13 +707,16 @@ void FConcertClientWorkspace::SetTransactionActivity(const FConcertSyncTransacti
 	if (LiveSession->GetSessionDatabase().SetTransactionActivity(InTransactionActivity))
 	{
 		PostActivityUpdated(InTransactionActivity);
-		if (TransactionManager)
-		{
-			TransactionManager->HandleRemoteTransaction(InTransactionActivity.EndpointId, InTransactionActivity.EventId, bHasSyncedWorkspace);
-		}
-		if (LiveTransactionAuthors)
+		if (!InTransactionActivity.bIgnored)
 		{
-			LiveTransactionAuthors->AddLiveTransactionActivity(InTransactionActivity.EndpointId, InTransactionActivity.EventData.Transaction.ModifiedPackages);
+			if (TransactionManager)
+			{
+				TransactionManager->HandleRemoteTransaction(InTransactionActivity.EndpointId, InTransactionActivity.EventId, bHasSyncedWorkspace);
+			}
+			if (LiveTransactionAuthors)
+			{
+				LiveTransactionAuthors->AddLiveTransactionActivity(InTransactionActivity.EndpointId, InTransactionActivity.EventData.Transaction.ModifiedPackages);
+			}
 		}
 	}
 	else
@@ -728,13 +731,16 @@ void FConcertClientWorkspace::SetPackageActivity(const FConcertSyncPackageActivi
 	if (LiveSession->GetSessionDatabase().SetPackageActivity(InPackageActivity))
 	{
 		PostActivityUpdated(InPackageActivity);
-		if (PackageManager)
-		{
-			PackageManager->HandleRemotePackage(InPackageActivity.EndpointId, InPackageActivity.EventId, bHasSyncedWorkspace);
-		}
-		if (LiveTransactionAuthors)
+		if (!InPackageActivity.bIgnored)
 		{
-			LiveTransactionAuthors->ResolveLiveTransactionAuthorsForPackage(InPackageActivity.EventData.Package.Info.PackageName);
+			if (PackageManager)
+			{
+				PackageManager->HandleRemotePackage(InPackageActivity.EndpointId, InPackageActivity.EventId, bHasSyncedWorkspace);
+			}
+			if (LiveTransactionAuthors)
+			{
+				LiveTransactionAuthors->ResolveLiveTransactionAuthorsForPackage(InPackageActivity.EventData.Package.Info.PackageName);
+			}
 		}
 	}
 	else
diff --git a/Engine/Plugins/Developer/Concert/ConcertSync/ConcertSyncCore/Source/ConcertSyncCore/Private/ConcertSyncSessionDatabase.cpp b/Engine/Plugins/Developer/Concert/ConcertSync/ConcertSyncCore/Source/ConcertSyncCore/Private/ConcertSyncSessionDatabase.cpp
--- a/Engine/Plugins/Developer/Concert/ConcertSync/ConcertSyncCore/Source/ConcertSyncCore/Private/ConcertSyncSessionDatabase.cpp
+++ b/Engine/Plugins/Developer/Concert/ConcertSync/ConcertSyncCore/Source/ConcertSyncCore/Private/ConcertSyncSessionDatabase.cpp
@@ -332,6 +332,10 @@ public:
 		PREPARE_STATEMENT(Statement_GetActivityIdAndEventTypesInRange);
 		PREPARE_STATEMENT(Statement_GetActivityMaxId);
 		
+		PREPARE_STATEMENT(Statement_IgnoreActivity);
+		PREPARE_STATEMENT(Statement_PerceiveActivity);
+		PREPARE_STATEMENT(Statement_IsActivityIgnored);
+
 		PREPARE_STATEMENT(Statement_MapObjectNameIdToLockEventId);
 		PREPARE_STATEMENT(Statement_UnmapObjectNameIdsForLockEventId);
 		PREPARE_STATEMENT(Statement_GetLockEventIdsForObjectNameId);
@@ -859,6 +863,35 @@ public:
 		return Statement_GetActivityMaxId.ExecuteSingle(OutActivityId);
 	}
 
+	/**
+	 * Statements working on ignored_activities
+	 */
+	
+	/** Add the activity_id to ignored_activities */
+	SQLITE_PREPARED_STATEMENT_BINDINGS_ONLY(FIgnoreActivity, "INSERT OR REPLACE INTO ignored_activities(activity_id) VALUES(?1);", SQLITE_PREPARED_STATEMENT_BINDINGS(int64));
+	FIgnoreActivity Statement_IgnoreActivity;
+	bool IgnoreActivity(const int64 InActivityId)
+	{
+		return Statement_IgnoreActivity.BindAndExecute(InActivityId);
+	}
+
+	/** Remove the activity_id from ignored_activities */
+	SQLITE_PREPARED_STATEMENT_BINDINGS_ONLY(FPerceiveActivity, "DELETE FROM ignored_activities WHERE activity_id = ?1;", SQLITE_PREPARED_STATEMENT_BINDINGS(int64));
+	FPerceiveActivity Statement_PerceiveActivity;
+	bool PerceiveActivity(const int64 InActivityId)
+	{
+		return Statement_PerceiveActivity.BindAndExecute(InActivityId);
+	}
+
+	/** See if the given activity_id is in ignored_activities */
+	SQLITE_PREPARED_STATEMENT(FIsActivityIgnored, "SELECT activity_id FROM ignored_activities WHERE activity_id = ?1;", SQLITE_PREPARED_STATEMENT_COLUMNS(int64), SQLITE_PREPARED_STATEMENT_BINDINGS(int64));
+	FIsActivityIgnored Statement_IsActivityIgnored;
+	bool IsActivityIgnored(const int64 InActivityId)
+	{
+		int64 OutActivityId = 0;
+		return Statement_IsActivityIgnored.BindAndExecuteSingle(InActivityId, OutActivityId);
+	}
+
 	/**
 	 * Statements working on resource_locks
 	 */
@@ -1139,6 +1172,7 @@ bool FConcertSyncSessionDatabase::Open(const FString& InSessionPath, const ESQLi
 	CREATE_TABLE("transaction_events", "transaction_event_id INTEGER PRIMARY KEY, data_filename TEXT NOT NULL");
 	CREATE_TABLE("package_events", "package_event_id INTEGER PRIMARY KEY, package_name_id INTEGER NOT NULL, package_revision INTEGER NOT NULL, package_info_size_bytes INTEGER NOT NULL, package_info_data BLOB, transaction_event_id_at_save INTEGER NOT NULL, data_filename TEXT NOT NULL, FOREIGN KEY(package_name_id) REFERENCES package_names(package_name_id)");
 	CREATE_TABLE("activities", "activity_id INTEGER PRIMARY KEY, endpoint_id BLOB NOT NULL, event_time INTEGER NOT NULL, event_type INTEGER NOT NULL, event_id INTEGER NOT NULL, event_summary_type TEXT NOT NULL, event_summary_size_bytes INTEGER NOT NULL, event_summary_data BLOB, FOREIGN KEY(endpoint_id) REFERENCES endpoints(endpoint_id)");
+	CREATE_TABLE("ignored_activities", "activity_id INTEGER NOT NULL, FOREIGN KEY(activity_id) REFERENCES activities(activity_id)");
 	CREATE_TABLE("resource_locks", "object_name_id INTEGER NOT NULL, lock_event_id INTEGER NOT NULL, FOREIGN KEY(object_name_id) REFERENCES object_names(object_name_id), FOREIGN KEY(lock_event_id) REFERENCES lock_events(lock_event_id)");
 	CREATE_TABLE("package_transactions", "package_name_id INTEGER NOT NULL, transaction_event_id INTEGER NOT NULL, FOREIGN KEY(package_name_id) REFERENCES package_names(package_name_id), FOREIGN KEY(transaction_event_id) REFERENCES transaction_events(transaction_event_id)");
 	CREATE_TABLE("object_transactions", "object_name_id INTEGER NOT NULL, transaction_event_id INTEGER NOT NULL, FOREIGN KEY(object_name_id) REFERENCES object_names(object_name_id), FOREIGN KEY(transaction_event_id) REFERENCES transaction_events(transaction_event_id)");
@@ -1161,6 +1195,7 @@ bool FConcertSyncSessionDatabase::Open(const FString& InSessionPath, const ESQLi
 	CREATE_UNIQUE_INDEX("idx_package_names_in_package_names", "package_names", "package_name");
 	CREATE_INDEX("idx_package_name_ids_in_package_events", "package_events", "package_name_id");
 	CREATE_INDEX("idx_event_ids_in_activities", "activities", "event_id");
+	CREATE_UNIQUE_INDEX("idx_activity_ids_in_ignored_activities", "ignored_activities", "activity_id");
 	CREATE_INDEX("idx_object_name_ids_in_resource_locks", "resource_locks", "object_name_id");
 	CREATE_INDEX("idx_lock_event_ids_in_resource_locks", "resource_locks", "lock_event_id");
 	CREATE_INDEX("idx_package_name_ids_in_package_transactions", "package_transactions", "package_name_id");
@@ -1227,7 +1262,8 @@ bool FConcertSyncSessionDatabase::AddConnectionActivity(const FConcertSyncConnec
 	FConcertSyncSessionDatabaseScopedTransaction ScopedTransaction(*Statements);
 	return ScopedTransaction.CommitOrRollback(
 		AddConnectionEvent(InConnectionActivity.EventData, OutConnectionEventId) &&
-		Statements->AddActivityData(InConnectionActivity.EndpointId, EConcertSyncActivityEventType::Connection, OutConnectionEventId, InConnectionActivity.EventSummary, OutActivityId)
+		Statements->AddActivityData(InConnectionActivity.EndpointId, EConcertSyncActivityEventType::Connection, OutConnectionEventId, InConnectionActivity.EventSummary, OutActivityId) &&
+		SetActivityIgnoredState(OutActivityId, InConnectionActivity.bIgnored)
 		);
 }
 
@@ -1236,7 +1272,8 @@ bool FConcertSyncSessionDatabase::AddLockActivity(const FConcertSyncLockActivity
 	FConcertSyncSessionDatabaseScopedTransaction ScopedTransaction(*Statements);
 	return ScopedTransaction.CommitOrRollback(
 		AddLockEvent(InLockActivity.EventData, OutLockEventId) &&
-		Statements->AddActivityData(InLockActivity.EndpointId, EConcertSyncActivityEventType::Lock, OutLockEventId, InLockActivity.EventSummary, OutActivityId)
+		Statements->AddActivityData(InLockActivity.EndpointId, EConcertSyncActivityEventType::Lock, OutLockEventId, InLockActivity.EventSummary, OutActivityId) &&
+		SetActivityIgnoredState(OutActivityId, InLockActivity.bIgnored)
 		);
 }
 
@@ -1245,7 +1282,8 @@ bool FConcertSyncSessionDatabase::AddTransactionActivity(const FConcertSyncTrans
 	FConcertSyncSessionDatabaseScopedTransaction ScopedTransaction(*Statements);
 	return ScopedTransaction.CommitOrRollback(
 		AddTransactionEvent(InTransactionActivity.EventData, OutTransactionEventId) &&
-		Statements->AddActivityData(InTransactionActivity.EndpointId, EConcertSyncActivityEventType::Transaction, OutTransactionEventId, InTransactionActivity.EventSummary, OutActivityId)
+		Statements->AddActivityData(InTransactionActivity.EndpointId, EConcertSyncActivityEventType::Transaction, OutTransactionEventId, InTransactionActivity.EventSummary, OutActivityId) &&
+		SetActivityIgnoredState(OutActivityId, InTransactionActivity.bIgnored)
 		);
 }
 
@@ -1254,7 +1292,8 @@ bool FConcertSyncSessionDatabase::AddPackageActivity(const FConcertSyncPackageAc
 	FConcertSyncSessionDatabaseScopedTransaction ScopedTransaction(*Statements);
 	return ScopedTransaction.CommitOrRollback(
 		AddPackageEvent(InPackageActivity.EventData, OutPackageEventId) &&
-		Statements->AddActivityData(InPackageActivity.EndpointId, EConcertSyncActivityEventType::Package, OutPackageEventId, InPackageActivity.EventSummary, OutActivityId)
+		Statements->AddActivityData(InPackageActivity.EndpointId, EConcertSyncActivityEventType::Package, OutPackageEventId, InPackageActivity.EventSummary, OutActivityId) &&
+		SetActivityIgnoredState(OutActivityId, InPackageActivity.bIgnored)
 		);
 }
 
@@ -1263,7 +1302,8 @@ bool FConcertSyncSessionDatabase::SetConnectionActivity(const FConcertSyncConnec
 	FConcertSyncSessionDatabaseScopedTransaction ScopedTransaction(*Statements);
 	return ScopedTransaction.CommitOrRollback(
 		SetConnectionEvent(InConnectionActivity.EventId, InConnectionActivity.EventData) &&
-		Statements->SetActivityData(InConnectionActivity.ActivityId, InConnectionActivity.EndpointId, InConnectionActivity.EventTime, InConnectionActivity.EventType, InConnectionActivity.EventId, InConnectionActivity.EventSummary)
+		Statements->SetActivityData(InConnectionActivity.ActivityId, InConnectionActivity.EndpointId, InConnectionActivity.EventTime, InConnectionActivity.EventType, InConnectionActivity.EventId, InConnectionActivity.EventSummary) &&
+		SetActivityIgnoredState(InConnectionActivity.ActivityId, InConnectionActivity.bIgnored)
 		);
 }
 
@@ -1272,7 +1312,8 @@ bool FConcertSyncSessionDatabase::SetLockActivity(const FConcertSyncLockActivity
 	FConcertSyncSessionDatabaseScopedTransaction ScopedTransaction(*Statements);
 	return ScopedTransaction.CommitOrRollback(
 		SetLockEvent(InLockActivity.EventId, InLockActivity.EventData) &&
-		Statements->SetActivityData(InLockActivity.ActivityId, InLockActivity.EndpointId, InLockActivity.EventTime, InLockActivity.EventType, InLockActivity.EventId, InLockActivity.EventSummary)
+		Statements->SetActivityData(InLockActivity.ActivityId, InLockActivity.EndpointId, InLockActivity.EventTime, InLockActivity.EventType, InLockActivity.EventId, InLockActivity.EventSummary) &&
+		SetActivityIgnoredState(InLockActivity.ActivityId, InLockActivity.bIgnored)
 		);
 }
 
@@ -1281,7 +1322,8 @@ bool FConcertSyncSessionDatabase::SetTransactionActivity(const FConcertSyncTrans
 	FConcertSyncSessionDatabaseScopedTransaction ScopedTransaction(*Statements);
 	return ScopedTransaction.CommitOrRollback(
 		SetTransactionEvent(InTransactionActivity.EventId, InTransactionActivity.EventData, bMetaDataOnly) &&
-		Statements->SetActivityData(InTransactionActivity.ActivityId, InTransactionActivity.EndpointId, InTransactionActivity.EventTime, InTransactionActivity.EventType, InTransactionActivity.EventId, InTransactionActivity.EventSummary)
+		Statements->SetActivityData(InTransactionActivity.ActivityId, InTransactionActivity.EndpointId, InTransactionActivity.EventTime, InTransactionActivity.EventType, InTransactionActivity.EventId, InTransactionActivity.EventSummary) &&
+		SetActivityIgnoredState(InTransactionActivity.ActivityId, InTransactionActivity.bIgnored)
 		);
 }
 
@@ -1290,13 +1332,15 @@ bool FConcertSyncSessionDatabase::SetPackageActivity(const FConcertSyncPackageAc
 	FConcertSyncSessionDatabaseScopedTransaction ScopedTransaction(*Statements);
 	return ScopedTransaction.CommitOrRollback(
 		SetPackageEvent(InPackageActivity.EventId, InPackageActivity.EventData, bMetaDataOnly) &&
-		Statements->SetActivityData(InPackageActivity.ActivityId, InPackageActivity.EndpointId, InPackageActivity.EventTime, InPackageActivity.EventType, InPackageActivity.EventId, InPackageActivity.EventSummary)
+		Statements->SetActivityData(InPackageActivity.ActivityId, InPackageActivity.EndpointId, InPackageActivity.EventTime, InPackageActivity.EventType, InPackageActivity.EventId, InPackageActivity.EventSummary) &&
+		SetActivityIgnoredState(InPackageActivity.ActivityId, InPackageActivity.bIgnored)
 		);
 }
 
 bool FConcertSyncSessionDatabase::GetActivity(const int64 InActivityId, FConcertSyncActivity& OutActivity) const
 {
 	OutActivity.ActivityId = InActivityId;
+	OutActivity.bIgnored = Statements->IsActivityIgnored(InActivityId);
 	return Statements->GetActivityDataForId(InActivityId, OutActivity.EndpointId, OutActivity.EventTime, OutActivity.EventType, OutActivity.EventId, OutActivity.EventSummary);
 }
 
@@ -1333,7 +1377,12 @@ bool FConcertSyncSessionDatabase::GetActivityForEvent(const int64 InEventId, con
 {
 	OutActivity.EventId = InEventId;
 	OutActivity.EventType = InEventType;
-	return Statements->GetActivityDataForEvent(InEventId, InEventType, OutActivity.ActivityId, OutActivity.EndpointId, OutActivity.EventTime, OutActivity.EventSummary);
+	if (Statements->GetActivityDataForEvent(InEventId, InEventType, OutActivity.ActivityId, OutActivity.EndpointId, OutActivity.EventTime, OutActivity.EventSummary))
+	{
+		OutActivity.bIgnored = Statements->IsActivityIgnored(OutActivity.ActivityId);
+		return true;
+	}
+	return false;
 }
 
 bool FConcertSyncSessionDatabase::GetConnectionActivityForEvent(const int64 InConnectionEventId, FConcertSyncConnectionActivity& OutConnectionActivity) const
@@ -1362,10 +1411,11 @@ bool FConcertSyncSessionDatabase::GetPackageActivityForEvent(const int64 InPacka
 
 bool FConcertSyncSessionDatabase::EnumerateActivities(TFunctionRef<bool(FConcertSyncActivity&&)> InCallback) const
 {
-	return Statements->GetAllActivityData([&InCallback](const int64 InActivityId, const FGuid& InEndpointId, const FDateTime InEventTime, const EConcertSyncActivityEventType InEventType, const int64 InEventId, FConcertSessionSerializedCborPayload&& InEventSummary)
+	return Statements->GetAllActivityData([this, &InCallback](const int64 InActivityId, const FGuid& InEndpointId, const FDateTime InEventTime, const EConcertSyncActivityEventType InEventType, const int64 InEventId, FConcertSessionSerializedCborPayload&& InEventSummary)
 	{
 		FConcertSyncActivity Activity;
 		Activity.ActivityId = InActivityId;
+		Activity.bIgnored = Statements->IsActivityIgnored(InActivityId);
 		Activity.EndpointId = InEndpointId;
 		Activity.EventTime = InEventTime;
 		Activity.EventType = InEventType;
@@ -1383,6 +1433,7 @@ bool FConcertSyncSessionDatabase::EnumerateConnectionActivities(TFunctionRef<boo
 	{
 		FConcertSyncConnectionActivity ConnectionActivity;
 		ConnectionActivity.ActivityId = InActivityId;
+		ConnectionActivity.bIgnored = Statements->IsActivityIgnored(InActivityId);
 		ConnectionActivity.EndpointId = InEndpointId;
 		ConnectionActivity.EventTime = InEventTime;
 		ConnectionActivity.EventType = EConcertSyncActivityEventType::Connection;
@@ -1404,6 +1455,7 @@ bool FConcertSyncSessionDatabase::EnumerateLockActivities(TFunctionRef<bool(FCon
 	{
 		FConcertSyncLockActivity LockActivity;
 		LockActivity.ActivityId = InActivityId;
+		LockActivity.bIgnored = Statements->IsActivityIgnored(InActivityId);
 		LockActivity.EndpointId = InEndpointId;
 		LockActivity.EventTime = InEventTime;
 		LockActivity.EventType = EConcertSyncActivityEventType::Lock;
@@ -1425,6 +1477,7 @@ bool FConcertSyncSessionDatabase::EnumerateTransactionActivities(TFunctionRef<bo
 	{
 		FConcertSyncTransactionActivity TransactionActivity;
 		TransactionActivity.ActivityId = InActivityId;
+		TransactionActivity.bIgnored = Statements->IsActivityIgnored(InActivityId);
 		TransactionActivity.EndpointId = InEndpointId;
 		TransactionActivity.EventTime = InEventTime;
 		TransactionActivity.EventType = EConcertSyncActivityEventType::Transaction;
@@ -1446,6 +1499,7 @@ bool FConcertSyncSessionDatabase::EnumeratePackageActivities(TFunctionRef<bool(F
 	{
 		FConcertSyncPackageActivity PackageActivity;
 		PackageActivity.ActivityId = InActivityId;
+		PackageActivity.bIgnored = Statements->IsActivityIgnored(InActivityId);
 		PackageActivity.EndpointId = InEndpointId;
 		PackageActivity.EventTime = InEventTime;
 		PackageActivity.EventType = EConcertSyncActivityEventType::Package;
@@ -1463,10 +1517,11 @@ bool FConcertSyncSessionDatabase::EnumeratePackageActivities(TFunctionRef<bool(F
 
 bool FConcertSyncSessionDatabase::EnumerateActivitiesForEventType(const EConcertSyncActivityEventType InEventType, TFunctionRef<bool(FConcertSyncActivity&&)> InCallback) const
 {
-	return Statements->GetAllActivityDataForEventType(InEventType, [InEventType, &InCallback](const int64 InActivityId, const FGuid& InEndpointId, const FDateTime InEventTime, const int64 InEventId, FConcertSessionSerializedCborPayload&& InEventSummary)
+	return Statements->GetAllActivityDataForEventType(InEventType, [this, InEventType, &InCallback](const int64 InActivityId, const FGuid& InEndpointId, const FDateTime InEventTime, const int64 InEventId, FConcertSessionSerializedCborPayload&& InEventSummary)
 	{
 		FConcertSyncActivity Activity;
 		Activity.ActivityId = InActivityId;
+		Activity.bIgnored = Statements->IsActivityIgnored(InActivityId);
 		Activity.EndpointId = InEndpointId;
 		Activity.EventTime = InEventTime;
 		Activity.EventType = InEventType;
@@ -1480,10 +1535,11 @@ bool FConcertSyncSessionDatabase::EnumerateActivitiesForEventType(const EConcert
 
 bool FConcertSyncSessionDatabase::EnumerateActivitiesInRange(const int64 InFirstActivityId, const int64 InMaxNumActivities, TFunctionRef<bool(FConcertSyncActivity&&)> InCallback) const
 {
-	return Statements->GetActivityDataInRange(InFirstActivityId, InMaxNumActivities, [&InCallback](const int64 InActivityId, const FGuid& InEndpointId, const FDateTime InEventTime, const EConcertSyncActivityEventType InEventType, const int64 InEventId, FConcertSessionSerializedCborPayload&& InEventSummary)
+	return Statements->GetActivityDataInRange(InFirstActivityId, InMaxNumActivities, [this, &InCallback](const int64 InActivityId, const FGuid& InEndpointId, const FDateTime InEventTime, const EConcertSyncActivityEventType InEventType, const int64 InEventId, FConcertSessionSerializedCborPayload&& InEventSummary)
 	{
 		FConcertSyncActivity Activity;
 		Activity.ActivityId = InActivityId;
+		Activity.bIgnored = Statements->IsActivityIgnored(InActivityId);
 		Activity.EndpointId = InEndpointId;
 		Activity.EventTime = InEventTime;
 		Activity.EventType = InEventType;
@@ -1553,6 +1609,13 @@ bool FConcertSyncSessionDatabase::EnumerateEndpointIds(TFunctionRef<bool(FGuid)>
 	});
 }
 
+bool FConcertSyncSessionDatabase::SetActivityIgnoredState(const int64 InActivityId, const bool InIsIgnored)
+{
+	return InIsIgnored
+		? Statements->IgnoreActivity(InActivityId)
+		: Statements->PerceiveActivity(InActivityId);
+}
+
 bool FConcertSyncSessionDatabase::AddConnectionEvent(const FConcertSyncConnectionEvent& InConnectionEvent, int64& OutConnectionEventId)
 {
 	return Statements->AddConnectionEvent(InConnectionEvent.ConnectionEventType, OutConnectionEventId);
diff --git a/Engine/Plugins/Developer/Concert/ConcertSync/ConcertSyncCore/Source/ConcertSyncCore/Public/ConcertSyncSessionDatabase.h b/Engine/Plugins/Developer/Concert/ConcertSync/ConcertSyncCore/Source/ConcertSyncCore/Public/ConcertSyncSessionDatabase.h
--- a/Engine/Plugins/Developer/Concert/ConcertSync/ConcertSyncCore/Source/ConcertSyncCore/Public/ConcertSyncSessionDatabase.h
+++ b/Engine/Plugins/Developer/Concert/ConcertSync/ConcertSyncCore/Source/ConcertSyncCore/Public/ConcertSyncSessionDatabase.h
@@ -584,6 +584,16 @@ public:
 	bool IsHeadRevisionPackageEvent(const int64 InPackageEventId, bool& OutIsHeadRevision) const;
 
 private:
+	/**
+	 * Set the active ignored state for the given activity.
+	 *
+	 * @param InActivityId				The ID of the activity to update.
+	 * @param InIsIgnored				True if this activity should be ignored, false otherwise.
+	 *
+	 * @return True if the ignored state was set, false otherwise.
+	 */
+	bool SetActivityIgnoredState(const int64 InActivityId, const bool InIsIgnored);
+
 	/**
 	 * Add a new connection event to this database, assigning it a connection event ID.
 	 *
diff --git a/Engine/Plugins/Developer/Concert/ConcertSync/ConcertSyncCore/Source/ConcertSyncCore/Public/ConcertSyncSessionTypes.h b/Engine/Plugins/Developer/Concert/ConcertSync/ConcertSyncCore/Source/ConcertSyncCore/Public/ConcertSyncSessionTypes.h
--- a/Engine/Plugins/Developer/Concert/ConcertSync/ConcertSyncCore/Source/ConcertSyncCore/Public/ConcertSyncSessionTypes.h
+++ b/Engine/Plugins/Developer/Concert/ConcertSync/ConcertSyncCore/Source/ConcertSyncCore/Public/ConcertSyncSessionTypes.h
@@ -132,6 +132,10 @@ struct FConcertSyncActivity
 	UPROPERTY()
 	int64 ActivityId = 0;
 
+	/** True if this activity is included for tracking purposes only, and should be ignored when applying an activity feed */
+	UPROPERTY()
+	bool bIgnored = false;
+
 	/** The ID of the endpoint that produced the activity */
 	UPROPERTY()
 	FGuid EndpointId;
diff --git a/Engine/Plugins/Developer/Concert/ConcertSync/ConcertSyncServer/Source/ConcertSyncServer/Private/ConcertSyncServer.cpp b/Engine/Plugins/Developer/Concert/ConcertSync/ConcertSyncServer/Source/ConcertSyncServer/Private/ConcertSyncServer.cpp
--- a/Engine/Plugins/Developer/Concert/ConcertSync/ConcertSyncServer/Source/ConcertSyncServer/Private/ConcertSyncServer.cpp
+++ b/Engine/Plugins/Developer/Concert/ConcertSync/ConcertSyncServer/Source/ConcertSyncServer/Private/ConcertSyncServer.cpp
@@ -127,9 +127,12 @@ bool MigrateSessionData(const FConcertSyncSessionDatabase& InSourceDatabase, con
 				{
 					MIGRATE_SET_ERROR_RESULT_AND_RETURN("Failed to get connection activity '%s' from database at '%s': %s", *LexToString(InActivityId), *InSourceDatabase.GetFilename(), *InSourceDatabase.GetLastError());
 				}
-				if (!DestDatabase.SetConnectionActivity(ConnectionActivity))
+				if (InDestSessionFilter.bIncludeIgnoredActivities || !ConnectionActivity.bIgnored)
 				{
-					MIGRATE_SET_ERROR_RESULT_AND_RETURN("Failed to set connection activity '%s' on database at '%s': %s", *LexToString(InActivityId), *DestDatabase.GetFilename(), *DestDatabase.GetLastError());
+					if (!DestDatabase.SetConnectionActivity(ConnectionActivity))
+					{
+						MIGRATE_SET_ERROR_RESULT_AND_RETURN("Failed to set connection activity '%s' on database at '%s': %s", *LexToString(InActivityId), *DestDatabase.GetFilename(), *DestDatabase.GetLastError());
+					}
 				}
 			}
 			break;
@@ -141,9 +144,12 @@ bool MigrateSessionData(const FConcertSyncSessionDatabase& InSourceDatabase, con
 				{
 					MIGRATE_SET_ERROR_RESULT_AND_RETURN("Failed to get lock activity '%s' from database at '%s': %s", *LexToString(InActivityId), *InSourceDatabase.GetFilename(), *InSourceDatabase.GetLastError());
 				}
-				if (!DestDatabase.SetLockActivity(LockActivity))
+				if (InDestSessionFilter.bIncludeIgnoredActivities || !LockActivity.bIgnored)
 				{
-					MIGRATE_SET_ERROR_RESULT_AND_RETURN("Failed to set lock activity '%s' on database at '%s': %s", *LexToString(InActivityId), *DestDatabase.GetFilename(), *DestDatabase.GetLastError());
+					if (!DestDatabase.SetLockActivity(LockActivity))
+					{
+						MIGRATE_SET_ERROR_RESULT_AND_RETURN("Failed to set lock activity '%s' on database at '%s': %s", *LexToString(InActivityId), *DestDatabase.GetFilename(), *DestDatabase.GetLastError());
+					}
 				}
 			}
 			break;
@@ -155,7 +161,7 @@ bool MigrateSessionData(const FConcertSyncSessionDatabase& InSourceDatabase, con
 				{
 					MIGRATE_SET_ERROR_RESULT_AND_RETURN("Failed to get transaction activity '%s' from database at '%s': %s", *LexToString(InActivityId), *InSourceDatabase.GetFilename(), *InSourceDatabase.GetLastError());
 				}
-				if (ConcertSyncSessionDatabaseFilterUtil::TransactionEventPassesFilter(TransactionActivity.EventId, InDestSessionFilter, InSourceDatabase))
+				if ((InDestSessionFilter.bIncludeIgnoredActivities || !TransactionActivity.bIgnored) && ConcertSyncSessionDatabaseFilterUtil::TransactionEventPassesFilter(TransactionActivity.EventId, InDestSessionFilter, InSourceDatabase))
 				{
 					if (!InSourceDatabase.GetTransactionEvent(TransactionActivity.EventId, TransactionActivity.EventData, InDestSessionFilter.bMetaDataOnly))
 					{
@@ -176,7 +182,7 @@ bool MigrateSessionData(const FConcertSyncSessionDatabase& InSourceDatabase, con
 				{
 					MIGRATE_SET_ERROR_RESULT_AND_RETURN("Failed to get package activity '%s' from database at '%s': %s", *LexToString(InActivityId), *InSourceDatabase.GetFilename(), *InSourceDatabase.GetLastError());
 				}
-				if (ConcertSyncSessionDatabaseFilterUtil::PackageEventPassesFilter(PackageActivity.EventId, InDestSessionFilter, InSourceDatabase))
+				if ((InDestSessionFilter.bIncludeIgnoredActivities || !PackageActivity.bIgnored) && ConcertSyncSessionDatabaseFilterUtil::PackageEventPassesFilter(PackageActivity.EventId, InDestSessionFilter, InSourceDatabase))
 				{
 					if (!InSourceDatabase.GetPackageEvent(PackageActivity.EventId, PackageActivity.EventData, InDestSessionFilter.bMetaDataOnly))
 					{
