Use the SQLite user_version meta-data to store the schema version

diff --git a/Engine/Plugins/Developer/Concert/ConcertSync/ConcertSyncCore/Source/ConcertSyncCore/Private/ConcertSyncSessionDatabase.cpp b/Engine/Plugins/Developer/Concert/ConcertSync/ConcertSyncCore/Source/ConcertSyncCore/Private/ConcertSyncSessionDatabase.cpp
--- a/Engine/Plugins/Developer/Concert/ConcertSync/ConcertSyncCore/Source/ConcertSyncCore/Private/ConcertSyncSessionDatabase.cpp
+++ b/Engine/Plugins/Developer/Concert/ConcertSync/ConcertSyncCore/Source/ConcertSyncCore/Private/ConcertSyncSessionDatabase.cpp
@@ -264,6 +264,14 @@ bool ReadPackage(const TArray<uint8>& InSerializedPackageData, FConcertPackageIn
 
 } // namespace PackageDataUtil
 
+enum class FConcertSyncSessionDatabaseVersion
+{
+	Empty = 0,
+	Initial = 1,
+
+	Current = Initial,
+};
+
 class FConcertSyncSessionDatabaseStatements
 {
 public:
@@ -1157,6 +1165,15 @@ bool FConcertSyncSessionDatabase::Open(const FString& InSessionPath, const ESQLi
 	Database->Execute(TEXT("PRAGMA journal_mode=WAL;"));
 	Database->Execute(TEXT("PRAGMA synchronous=NORMAL;"));
 
+	int32 LoadedDatabaseVersion = 0;
+	Database->GetUserVersion(LoadedDatabaseVersion);
+	if (LoadedDatabaseVersion > (int32)FConcertSyncSessionDatabaseVersion::Current)
+	{
+		Close();
+		UE_LOG(LogConcert, Error, TEXT("Failed to open session database for '%s': Database is too new (version %d, expected <= %d)"), *InSessionPath, LoadedDatabaseVersion, (int32)FConcertSyncSessionDatabaseVersion::Current);
+		return false;
+	}
+
 	// Create our required tables
 #define CREATE_TABLE(NAME, STATEMENT)																										\
 	if (!Database->Execute(TEXT("CREATE TABLE IF NOT EXISTS ") TEXT(NAME) TEXT("(") TEXT(STATEMENT) TEXT(");")))							\
@@ -1205,6 +1222,13 @@ bool FConcertSyncSessionDatabase::Open(const FString& InSessionPath, const ESQLi
 #undef CREATE_INDEX
 #undef CREATE_UNIQUE_INDEX
 
+	// The database will have the latest schema at this point, so update the user-version
+	if (!Database->SetUserVersion((int32)FConcertSyncSessionDatabaseVersion::Current))
+	{
+		Close();
+		return false;
+	}
+
 	// Create our required prepared statements
 	Statements = MakeUnique<FConcertSyncSessionDatabaseStatements>(*Database);
 	if (!Statements->CreatePreparedStatements())
diff --git a/Engine/Plugins/Runtime/Database/SQLiteCore/Source/SQLiteCore/Private/SQLiteDatabase.cpp b/Engine/Plugins/Runtime/Database/SQLiteCore/Source/SQLiteCore/Private/SQLiteDatabase.cpp
--- a/Engine/Plugins/Runtime/Database/SQLiteCore/Source/SQLiteCore/Private/SQLiteDatabase.cpp
+++ b/Engine/Plugins/Runtime/Database/SQLiteCore/Source/SQLiteCore/Private/SQLiteDatabase.cpp
@@ -117,6 +117,34 @@ FString FSQLiteDatabase::GetFilename() const
 	return FString();
 }
 
+bool FSQLiteDatabase::GetApplicationId(int32& OutApplicationId) const
+{
+	return const_cast<FSQLiteDatabase*>(this)->Execute(TEXT("PRAGMA application_id;"), [&OutApplicationId](const FSQLitePreparedStatement& InStatement)
+	{
+		InStatement.GetColumnValueByIndex(0, OutApplicationId);
+		return ESQLitePreparedStatementExecuteRowResult::Stop;
+	});
+}
+
+bool FSQLiteDatabase::SetApplicationId(const int32 InApplicationId)
+{
+	return Execute(*FString::Printf(TEXT("PRAGMA application_id = %d;"), InApplicationId));
+}
+
+bool FSQLiteDatabase::GetUserVersion(int32& OutUserVersion) const
+{
+	return const_cast<FSQLiteDatabase*>(this)->Execute(TEXT("PRAGMA user_version;"), [&OutUserVersion](const FSQLitePreparedStatement& InStatement)
+	{
+		InStatement.GetColumnValueByIndex(0, OutUserVersion);
+		return ESQLitePreparedStatementExecuteRowResult::Stop;
+	});
+}
+
+bool FSQLiteDatabase::SetUserVersion(const int32 InUserVersion)
+{
+	return Execute(*FString::Printf(TEXT("PRAGMA user_version = %d;"), InUserVersion));
+}
+
 bool FSQLiteDatabase::Execute(const TCHAR* InStatement)
 {
 	if (!Database)
diff --git a/Engine/Plugins/Runtime/Database/SQLiteCore/Source/SQLiteCore/Public/SQLiteDatabase.h b/Engine/Plugins/Runtime/Database/SQLiteCore/Source/SQLiteCore/Public/SQLiteDatabase.h
--- a/Engine/Plugins/Runtime/Database/SQLiteCore/Source/SQLiteCore/Public/SQLiteDatabase.h
+++ b/Engine/Plugins/Runtime/Database/SQLiteCore/Source/SQLiteCore/Public/SQLiteDatabase.h
@@ -60,6 +60,32 @@ public:
 	 */
 	FString GetFilename() const;
 
+	/**
+	 * Get the application ID set in the database header.
+	 * @note A list of assigned application IDs can be seen by consulting the magic.txt file in the SQLite source repository.
+	 * @return true if the get was a success.
+	 */
+	bool GetApplicationId(int32& OutApplicationId) const;
+
+	/**
+	 * Set the application ID in the database header.
+	 * @note A list of assigned application IDs can be seen by consulting the magic.txt file in the SQLite source repository.
+	 * @return true if the set was a success.
+	 */
+	bool SetApplicationId(const int32 InApplicationId);
+
+	/**
+	 * Get the user version set in the database header.
+	 * @return true if the get was a success.
+	 */
+	bool GetUserVersion(int32& OutUserVersion) const;
+
+	/**
+	 * Set the user version in the database header.
+	 * @return true if the set was a success.
+	 */
+	bool SetUserVersion(const int32 InUserVersion);
+
 	/**
 	 * Execute a statement that requires no result state.
 	 * @note For statements that require a result, or that you wish to reuse repeatedly (including using bindings), you should consider using FSQLitePreparedStatement.
