Improved scalability of Disaster Recovery
  - Converted Concert API transferring package data in-memory only model to a streaming model to support packages bigger than 2 GB. (TNumberiLimit<int32>::max())
  - Added the IConcertFileSharing interface to share large files between the client and the server. This is used as a side channel to the Concert request/response and event protocol.
  - Fixed the ConcertClientPackageManager to prevent sending the package data for each the 'pre-save' when the 'live sync'  is off. It only emits it once.
  - Fixed UI to correctly report pre-save vs save vs auto-save for package as well as when a package is discarded.

diff --git a/Engine/Plugins/Developer/Concert/ConcertMain/Source/Concert/Private/ConcertUtil.cpp b/Engine/Plugins/Developer/Concert/ConcertMain/Source/Concert/Private/ConcertUtil.cpp
--- a/Engine/Plugins/Developer/Concert/ConcertMain/Source/Concert/Private/ConcertUtil.cpp
+++ b/Engine/Plugins/Developer/Concert/ConcertMain/Source/Concert/Private/ConcertUtil.cpp
@@ -23,3 +23,29 @@ bool ConcertUtil::DeleteDirectoryTree(const TCHAR* InDirectoryToDelete, const TC
 
 	return true;
 }
+
+bool ConcertUtil::Copy(FArchive& DstAr, FArchive& SrcAr, int64 Size)
+{
+	check(DstAr.IsSaving() && SrcAr.IsLoading());
+
+	constexpr int64 DataChunkSize = 4 * 1024; // 4K is a typical low level block size.
+	TArray<uint8> DataChunk;
+	DataChunk.AddUninitialized(static_cast<int32>(DataChunkSize));
+
+	// Copy full 4K block from the source to dest.
+	int64 ChunkCount = Size / DataChunkSize;
+	while (ChunkCount > 0)
+	{
+		SrcAr.Serialize(DataChunk.GetData(), DataChunkSize); // Read.
+		DstAr.Serialize(DataChunk.GetData(), DataChunkSize); // Write.
+		--ChunkCount;
+	}
+
+	// Copy the last few bytes remaining.
+	int64 Remaining = Size % DataChunkSize;
+	SrcAr.Serialize(DataChunk.GetData(), Remaining); // Read.
+	DstAr.Serialize(DataChunk.GetData(), Remaining); // Write.
+
+	return !SrcAr.IsError() && !DstAr.IsError();
+}
+
diff --git a/Engine/Plugins/Developer/Concert/ConcertMain/Source/Concert/Public/ConcertUtil.h b/Engine/Plugins/Developer/Concert/ConcertMain/Source/Concert/Public/ConcertUtil.h
--- a/Engine/Plugins/Developer/Concert/ConcertMain/Source/Concert/Public/ConcertUtil.h
+++ b/Engine/Plugins/Developer/Concert/ConcertMain/Source/Concert/Public/ConcertUtil.h
@@ -10,4 +10,7 @@ namespace ConcertUtil
 /** Delete a directory tree via a move and delete */
 CONCERT_API bool DeleteDirectoryTree(const TCHAR* InDirectoryToDelete, const TCHAR* InMoveToDirBeforeDelete = nullptr);
 
+/** Copy the specified data size from a source archive into a destination archive. */
+CONCERT_API bool Copy(FArchive& DstAr, FArchive& SrcAr, int64 Size);
+
 } // namespace ConcertUtil
diff --git a/Engine/Plugins/Developer/Concert/ConcertSync/ConcertSyncClient/Source/ConcertSyncClient/Private/ConcertActivityStream.cpp b/Engine/Plugins/Developer/Concert/ConcertSync/ConcertSyncClient/Source/ConcertSyncClient/Private/ConcertActivityStream.cpp
--- a/Engine/Plugins/Developer/Concert/ConcertSync/ConcertSyncClient/Source/ConcertSyncClient/Private/ConcertActivityStream.cpp
+++ b/Engine/Plugins/Developer/Concert/ConcertSync/ConcertSyncClient/Source/ConcertSyncClient/Private/ConcertActivityStream.cpp
@@ -101,6 +101,7 @@ void FConcertActivityStream::Fetch(int64 RequestFetchCount)
 						SyncActivity = SyncPackageActivity;
 						EventPayload = MakeUnique<FConcertSessionSerializedPayload>();
 						EventPayload->SetTypedPayload(SyncPackageActivity->EventData);
+						checkf(!SyncPackageActivity->EventData.Package.HasPackageData(), TEXT("The concert activity stream should only stream the package meta data. The package data itself should not be required where ConcertActivityStream is used. (The package data can be very large)"));
 					}
 					else // bIncludeActivityDetails was false/Other type of activities (lock/connection) don't have interesting data to inspect other than the one already provided by the generic activity.
 					{
diff --git a/Engine/Plugins/Developer/Concert/ConcertSync/ConcertSyncClient/Source/ConcertSyncClient/Private/ConcertClientPackageBridge.cpp b/Engine/Plugins/Developer/Concert/ConcertSync/ConcertSyncClient/Source/ConcertSyncClient/Private/ConcertClientPackageBridge.cpp
--- a/Engine/Plugins/Developer/Concert/ConcertSync/ConcertSyncClient/Source/ConcertSyncClient/Private/ConcertClientPackageBridge.cpp
+++ b/Engine/Plugins/Developer/Concert/ConcertSync/ConcertSyncClient/Source/ConcertSyncClient/Private/ConcertClientPackageBridge.cpp
@@ -150,14 +150,14 @@ void FConcertClientPackageBridge::HandlePackagePreSave(UPackage* Package)
 	FString PackageFilename;
 	if (FPackageName::TryConvertLongPackageNameToFilename(Package->GetFName().ToString(), PackageFilename, World ? FPackageName::GetMapPackageExtension() : FPackageName::GetAssetPackageExtension()))
 	{
-		FConcertPackage Event;
-		ConcertClientPackageBridgeUtil::FillPackageInfo(Package, EConcertPackageUpdateType::Saved, Event.Info);
-		Event.Info.bPreSave = true;
-		Event.Info.bAutoSave = GEngine->IsAutosaving();
-
-		if (FFileHelper::LoadFileToArray(Event.PackageData, *PackageFilename))
+		if (IFileManager::Get().FileExists(*PackageFilename))
 		{
-			OnLocalPackageEventDelegate.Broadcast(Event);
+			FConcertPackageInfo PackageInfo;
+			ConcertClientPackageBridgeUtil::FillPackageInfo(Package, EConcertPackageUpdateType::Saved, PackageInfo);
+			PackageInfo.bPreSave = true;
+			PackageInfo.bAutoSave = GEngine->IsAutosaving();
+		
+			OnLocalPackageEventDelegate.Broadcast(PackageInfo, PackageFilename);
 		}
 	}
 
@@ -191,15 +191,15 @@ void FConcertClientPackageBridge::HandlePackageSaved(const FString& PackageFilen
 	FName NewPackageName;
 	PackagesBeingRenamed.RemoveAndCopyValue(Package->GetFName(), NewPackageName);
 
-	FConcertPackage Event;
-	ConcertClientPackageBridgeUtil::FillPackageInfo(Package, NewPackageName.IsNone() ? EConcertPackageUpdateType::Saved : EConcertPackageUpdateType::Renamed, Event.Info);
-	Event.Info.NewPackageName = NewPackageName;
-	Event.Info.bPreSave = false;
-	Event.Info.bAutoSave = GEngine->IsAutosaving();
-
-	if (FFileHelper::LoadFileToArray(Event.PackageData, *PackageFilename))
+	if (IFileManager::Get().FileExists(*PackageFilename))
 	{
-		OnLocalPackageEventDelegate.Broadcast(Event);
+		FConcertPackageInfo PackageInfo;
+		ConcertClientPackageBridgeUtil::FillPackageInfo(Package, NewPackageName.IsNone() ? EConcertPackageUpdateType::Saved : EConcertPackageUpdateType::Renamed, PackageInfo);
+		PackageInfo.NewPackageName = NewPackageName;
+		PackageInfo.bPreSave = false;
+		PackageInfo.bAutoSave = GEngine->IsAutosaving();
+	
+		OnLocalPackageEventDelegate.Broadcast(PackageInfo, PackageFilename);
 	}
 
 	UE_LOG(LogConcert, Verbose, TEXT("Asset Saved: %s"), *Package->GetName());
@@ -233,15 +233,14 @@ void FConcertClientPackageBridge::HandleAssetAdded(UObject *Object)
 			// Saving the newly added asset here shouldn't modify any of its package flags since it's a 'dummy' save i.e. PKG_NewlyCreated
 			Package->SetPackageFlagsTo(PackageFlags);
 
-			FConcertPackage Event;
-			ConcertClientPackageBridgeUtil::FillPackageInfo(Package, EConcertPackageUpdateType::Added, Event.Info);
-
-			if (FFileHelper::LoadFileToArray(Event.PackageData, *PackageFilename))
+			if (IFileManager::Get().FileExists(*PackageFilename))
 			{
-				OnLocalPackageEventDelegate.Broadcast(Event);
-			}
+				FConcertPackageInfo PackageInfo;
+				ConcertClientPackageBridgeUtil::FillPackageInfo(Package, EConcertPackageUpdateType::Added, PackageInfo);
 
-			IFileManager::Get().Delete(*PackageFilename);
+				OnLocalPackageEventDelegate.Broadcast(PackageInfo, PackageFilename);
+				IFileManager::Get().Delete(*PackageFilename);
+			}
 		}
 	}
 
@@ -258,9 +257,9 @@ void FConcertClientPackageBridge::HandleAssetDeleted(UObject *Object)
 
 	UPackage* Package = Object->GetOutermost();
 
-	FConcertPackage Event;
-	ConcertClientPackageBridgeUtil::FillPackageInfo(Package, EConcertPackageUpdateType::Deleted, Event.Info);
-	OnLocalPackageEventDelegate.Broadcast(Event);
+	FConcertPackageInfo PackageInfo;
+	ConcertClientPackageBridgeUtil::FillPackageInfo(Package, EConcertPackageUpdateType::Deleted, PackageInfo);
+	OnLocalPackageEventDelegate.Broadcast(PackageInfo, FString());
 
 	UE_LOG(LogConcert, Verbose, TEXT("Asset Deleted: %s"), *Package->GetName());
 }
diff --git a/Engine/Plugins/Developer/Concert/ConcertSync/ConcertSyncClient/Source/ConcertSyncClient/Private/ConcertClientPackageManager.cpp b/Engine/Plugins/Developer/Concert/ConcertSync/ConcertSyncClient/Source/ConcertSyncClient/Private/ConcertClientPackageManager.cpp
--- a/Engine/Plugins/Developer/Concert/ConcertSync/ConcertSyncClient/Source/ConcertSyncClient/Private/ConcertClientPackageManager.cpp
+++ b/Engine/Plugins/Developer/Concert/ConcertSync/ConcertSyncClient/Source/ConcertSyncClient/Private/ConcertClientPackageManager.cpp
@@ -2,6 +2,7 @@
 
 #include "ConcertClientPackageManager.h"
 #include "IConcertSession.h"
+#include "IConcertFileSharingService.h"
 #include "ConcertSyncClientLiveSession.h"
 #include "ConcertSyncSessionDatabase.h"
 #include "ConcertLogGlobal.h"
@@ -9,6 +10,7 @@
 #include "ConcertWorkspaceMessages.h"
 #include "ConcertSandboxPlatformFile.h"
 #include "ConcertSyncClientUtil.h"
+#include "ConcertUtil.h"
 
 #include "Engine/World.h"
 #include "Engine/Engine.h"
@@ -27,10 +29,11 @@
 
 #define LOCTEXT_NAMESPACE "ConcertClientPackageManager"
 
-FConcertClientPackageManager::FConcertClientPackageManager(TSharedRef<FConcertSyncClientLiveSession> InLiveSession, IConcertClientPackageBridge* InPackageBridge)
-	: LiveSession(InLiveSession)
+FConcertClientPackageManager::FConcertClientPackageManager(TSharedRef<FConcertSyncClientLiveSession> InLiveSession, IConcertClientPackageBridge* InPackageBridge, TSharedPtr<IConcertFileSharingService> InFileSharingService)
+	: LiveSession(MoveTemp(InLiveSession))
 	, PackageBridge(InPackageBridge)
 	, bIgnorePackageDirtyEvent(false)
+	, FileSharingService(MoveTemp(InFileSharingService))
 {
 	check(LiveSession->IsValidSession());
 	check(EnumHasAnyFlags(LiveSession->GetSessionFlags(), EConcertSyncSessionFlags::EnablePackages));
@@ -133,7 +136,7 @@ void FConcertClientPackageManager::SynchronizePersistedFiles(const TMap<FString,
 		auto GetPackageFilenameForRevision = [this](const FString& PackageName, const int64 PackageRevision, FString& OutPackageFilename) -> bool
 		{
 			FConcertPackageInfo PackageInfo;
-			if (LiveSession->GetSessionDatabase().GetPackageDataForRevision(*PackageName, PackageInfo, nullptr, &PackageRevision))
+			if (LiveSession->GetSessionDatabase().GetPackageInfoForRevision(*PackageName, PackageInfo, &PackageRevision))
 			{
 				if (FPackageName::TryConvertLongPackageNameToFilename(PackageName, OutPackageFilename, PackageInfo.PackageFileExtension))
 				{
@@ -222,18 +225,17 @@ void FConcertClientPackageManager::HandleRemotePackage(const FGuid& InSourceEndp
 		return;
 	}
 
-	FConcertSyncPackageEvent PackageEvent;
-	if (LiveSession->GetSessionDatabase().GetPackageEvent(InPackageEventId, PackageEvent))
+	LiveSession->GetSessionDatabase().GetPackageEvent(InPackageEventId, [this](FConcertSyncPackageEventData& PackageEvent)
 	{
-		ApplyPackageUpdate(PackageEvent.Package);
-	}
+		ApplyPackageUpdate(PackageEvent.MetaData.PackageInfo, PackageEvent.PackageDataStream);
+	});
 }
 
 void FConcertClientPackageManager::ApplyAllHeadPackageData()
 {
-	LiveSession->GetSessionDatabase().EnumerateHeadRevisionPackageData([this](FConcertPackage&& InPackage)
+	LiveSession->GetSessionDatabase().EnumerateHeadRevisionPackageData([this](const FConcertPackageInfo& InPackageInfo, FConcertPackageDataStream& InPackageDataStream)
 	{
-		ApplyPackageUpdate(InPackage);
+		ApplyPackageUpdate(InPackageInfo, InPackageDataStream);
 		return true;
 	});
 }
@@ -270,23 +272,23 @@ bool FConcertClientPackageManager::PersistSessionChanges(TArrayView<const FName>
 	return false;
 }
 
-void FConcertClientPackageManager::ApplyPackageUpdate(const FConcertPackage& InPackage)
+void FConcertClientPackageManager::ApplyPackageUpdate(const FConcertPackageInfo& InPackageInfo, FConcertPackageDataStream& InPackageDataStream)
 {
-	switch (InPackage.Info.PackageUpdateType)
+	switch (InPackageInfo.PackageUpdateType)
 	{
 	case EConcertPackageUpdateType::Dummy:
 	case EConcertPackageUpdateType::Added:
 	case EConcertPackageUpdateType::Saved:
-		SavePackageFile(InPackage);
+		SavePackageFile(InPackageInfo, InPackageDataStream);
 		break;
 
 	case EConcertPackageUpdateType::Renamed:
-		DeletePackageFile(InPackage);
-		SavePackageFile(InPackage);
+		DeletePackageFile(InPackageInfo);
+		SavePackageFile(InPackageInfo, InPackageDataStream);
 		break;
 
 	case EConcertPackageUpdateType::Deleted:
-		DeletePackageFile(InPackage);
+		DeletePackageFile(InPackageInfo);
 		break;
 
 	default:
@@ -297,11 +299,10 @@ void FConcertClientPackageManager::ApplyPackageUpdate(const FConcertPackage& InP
 void FConcertClientPackageManager::HandlePackageRejectedEvent(const FConcertSessionContext& InEventContext, const FConcertPackageRejectedEvent& InEvent)
 {
 	// Package update was rejected, restore the head-revision of the package
-	FConcertPackage Package;
-	if (LiveSession->GetSessionDatabase().GetPackageDataForRevision(InEvent.PackageName, Package))
+	LiveSession->GetSessionDatabase().GetPackageDataForRevision(InEvent.PackageName, [this](const FConcertPackageInfo& InPackageInfo, FConcertPackageDataStream& InPackageDataStream)
 	{
-		ApplyPackageUpdate(Package);
-	}
+		ApplyPackageUpdate(InPackageInfo, InPackageDataStream);
+	});
 }
 
 void FConcertClientPackageManager::HandlePackageDirtyStateChanged(UPackage* InPackage)
@@ -317,89 +318,129 @@ void FConcertClientPackageManager::HandlePackageDirtyStateChanged(UPackage* InPa
 	}
 }
 
-void FConcertClientPackageManager::HandleLocalPackageEvent(const FConcertPackage& Package)
+void FConcertClientPackageManager::HandleLocalPackageEvent(const FConcertPackageInfo& PackageInfo, const FString& PackagePathname)
 {
 	// Ignore unwanted saves
-	if (Package.Info.PackageUpdateType == EConcertPackageUpdateType::Saved)
+	if (PackageInfo.PackageUpdateType == EConcertPackageUpdateType::Saved)
 	{
-		if (Package.Info.bPreSave)
+		if (!FPackageName::IsValidLongPackageName(PackageInfo.PackageName.ToString())) // Auto-Save might save the template in /Temp/... which is an invalid long package name.
+		{
+			return;
+		}
+		else if (PackageInfo.bAutoSave && !EnumHasAnyFlags(LiveSession->GetSessionFlags(), EConcertSyncSessionFlags::ShouldSendPackageAutoSaves))
+		{
+			return;
+		}
+		else if (PackageInfo.bPreSave)
 		{
 			// Pre-save events are used to send the pristine package state of a package (if enabled), so make sure we don't already have a history for this package
 			FConcertPackageInfo ExistingPackageInfo;
-			if (!EnumHasAnyFlags(LiveSession->GetSessionFlags(), EConcertSyncSessionFlags::ShouldSendPackagePristineState) || LiveSession->GetSessionDatabase().GetPackageDataForRevision(Package.Info.PackageName, ExistingPackageInfo, nullptr))
+			if (!EnumHasAnyFlags(LiveSession->GetSessionFlags(), EConcertSyncSessionFlags::ShouldSendPackagePristineState) || LiveSession->GetSessionDatabase().GetPackageInfoForRevision(PackageInfo.PackageName, ExistingPackageInfo))
 			{
 				return;
 			}
-		}
-		else
-		{
-			// Save events may optionally exclude auto-saves
-			if (Package.Info.bAutoSave && !EnumHasAnyFlags(LiveSession->GetSessionFlags(), EConcertSyncSessionFlags::ShouldSendPackageAutoSaves))
+			// Without live sync feature, the local database is not maintainted after the original 'sync on join'. Package that were not in the original sync don't have a revision from the client point of view.
+			else if (!EnumHasAnyFlags(LiveSession->GetSessionFlags(), EConcertSyncSessionFlags::EnableLiveSync) && EmittedPristinePackages.Contains(PackageInfo.PackageName))
 			{
-				return;
+				return; // Prevent capturing the original package state at every pre-save.
 			}
 		}
+	}
 
-		if (!FPackageName::IsValidLongPackageName(Package.Info.PackageName.ToString())) // Auto-Save might save the template in /Temp/... which is an invalid long package name.
+	if (PackageInfo.PackageUpdateType == EConcertPackageUpdateType::Added && EnumHasAnyFlags(LiveSession->GetSessionFlags(), EConcertSyncSessionFlags::ShouldUsePackageSandbox))
+	{
+		// If this package was locally added and we're using a sandbox, also write it to the correct location on disk (which will be placed into the sandbox directory)
+		FString DstPackagePathname;
+		if (FPackageName::TryConvertLongPackageNameToFilename(PackageInfo.PackageName.ToString(), DstPackagePathname, PackageInfo.PackageFileExtension))
 		{
-			return;
+			if (IFileManager::Get().Copy(*DstPackagePathname, *PackagePathname) != ECopyResult::COPY_OK)
+			{
+				UE_LOG(LogConcert, Error, TEXT("Failed to copy package file '%s' to the sandbox"), *PackagePathname);
+			}
 		}
 	}
 
-	if (Package.Info.PackageUpdateType == EConcertPackageUpdateType::Added && EnumHasAnyFlags(LiveSession->GetSessionFlags(), EConcertSyncSessionFlags::ShouldUsePackageSandbox))
+	FConcertPackageUpdateEvent Event;
+
+	// Copy or link the packge data to the Concert event.
+	int64 PackageFileSize = PackagePathname.IsEmpty() ? -1 : IFileManager::Get().FileSize(*PackagePathname); // EConcertPackageUpdateType::Delete is emitted with an empty pathname.
+	if (PackageFileSize > 0)
 	{
-		// If this package was locally added and we're using a sandbox, also write it to the correct location on disk (which will be placed into the sandbox directory)
-		FString PackageFilename;
-		if (FPackageName::TryConvertLongPackageNameToFilename(Package.Info.PackageName.ToString(), PackageFilename, Package.Info.PackageFileExtension))
+		if (CanExchangePackageDataAsByteArray(static_cast<uint64>(PackageFileSize)))
+		{
+			// Embed the package data directly in the event.
+			if (!FFileHelper::LoadFileToArray(Event.Package.PackageData, *PackagePathname))
+			{
+				UE_LOG(LogConcert, Error, TEXT("Failed to load file data '%s' in memory"), *PackagePathname);
+				return;
+			}
+		}
+		else if (FileSharingService && EnumHasAnyFlags(LiveSession->GetSessionFlags(), EConcertSyncSessionFlags::EnableFileSharing))
 		{
-			FFileHelper::SaveArrayToFile(Package.PackageData, *PackageFilename);
+			// Publish a copy of the package data in the sharing service and set the corresponding file ID in the event.
+			if (!FileSharingService->Publish(PackagePathname, Event.Package.FileId))
+			{
+				UE_LOG(LogConcert, Error, TEXT("Failed to share a copy of package file '%s'"), *PackagePathname);
+				return;
+			}
+		}
+		else
+		{
+			// Notify the client about the file being too large to be emitted.
+			UE_LOG(LogConcert, Error, TEXT("Failed to handle local package file '%s'. The file is too big to be sent over the network."), *PackagePathname);
+			OnConcertClientPackageTooLargeError().Broadcast(Event.Package.Info, PackageFileSize, FConcertPackage::GetMaxPackageDataSizeEmbeddableAsByteArray());
+			return;
 		}
 	}
 
-	FConcertPackageUpdateEvent Event;
-	Event.Package = Package;
+	Event.Package.Info = PackageInfo;
 	LiveSession->GetSessionDatabase().GetTransactionMaxEventId(Event.Package.Info.TransactionEventIdAtSave);
 	LiveSession->GetSession().SendCustomEvent(Event, LiveSession->GetSession().GetSessionServerEndpointId(), EConcertMessageFlags::ReliableOrdered);
+
+	if (PackageInfo.bPreSave && EnumHasAnyFlags(LiveSession->GetSessionFlags(), EConcertSyncSessionFlags::ShouldSendPackagePristineState) && !EnumHasAnyFlags(LiveSession->GetSessionFlags(), EConcertSyncSessionFlags::EnableLiveSync))
+	{
+		EmittedPristinePackages.Add(PackageInfo.PackageName); // Prevent sending the original package state at every pre-save.
+	}
 }
 
-void FConcertClientPackageManager::SavePackageFile(const FConcertPackage& Package)
+void FConcertClientPackageManager::SavePackageFile(const FConcertPackageInfo& PackageInfo, FConcertPackageDataStream& InPackageDataStream)
 {
 	// This path should only be taken for non-cooked targets for now
 	check(!FPlatformProperties::RequiresCookedData());
 
-	if (Package.PackageData.Num() == 0)
+	if (!InPackageDataStream.DataAr || InPackageDataStream.DataSize == 0)
 	{
-		// If we have no package data set, then this was from a meta-data 
+		// If we have no package data set, then this was from a meta-data
 		// only package sync, so we have no new contents to write to disk
 		return;
 	}
 
-	FString PackageName = Package.Info.PackageName.ToString();
+	FString PackageName = PackageInfo.PackageName.ToString();
 	ConcertSyncClientUtil::FlushPackageLoading(PackageName);
 
 	// Convert long package name to filename
 	FString PackageFilename;
-	bool bSuccess = FPackageName::TryConvertLongPackageNameToFilename(PackageName, PackageFilename, Package.Info.PackageFileExtension);
+	bool bSuccess = FPackageName::TryConvertLongPackageNameToFilename(PackageName, PackageFilename, PackageInfo.PackageFileExtension);
 	if (bSuccess)
 	{
 		// Overwrite the file on disk
-		FPlatformFileManager::Get().GetPlatformFile().SetReadOnly(*PackageFilename, false);
-		bSuccess = FFileHelper::SaveArrayToFile(Package.PackageData, *PackageFilename);
+		TUniquePtr<FArchive> DstAr(IFileManager::Get().CreateFileWriter(*PackageFilename, FILEWRITE_EvenIfReadOnly));
+		bSuccess = DstAr && ConcertUtil::Copy(*DstAr, *InPackageDataStream.DataAr, InPackageDataStream.DataSize);
 	}
 
 	if (bSuccess)
 	{
-		PackagesPendingHotReload.Add(Package.Info.PackageName);
-		PackagesPendingPurge.Remove(Package.Info.PackageName);
+		PackagesPendingHotReload.Add(PackageInfo.PackageName);
+		PackagesPendingPurge.Remove(PackageInfo.PackageName);
 	}
 }
 
-void FConcertClientPackageManager::DeletePackageFile(const FConcertPackage& Package)
+void FConcertClientPackageManager::DeletePackageFile(const FConcertPackageInfo& PackageInfo)
 {
 	// This path should only be taken for non-cooked targets for now
 	check(!FPlatformProperties::RequiresCookedData());
 
-	FString PackageName = Package.Info.PackageName.ToString();
+	FString PackageName = PackageInfo.PackageName.ToString();
 	ConcertSyncClientUtil::FlushPackageLoading(PackageName);
 
 	// Convert long package name to filename
@@ -420,8 +461,8 @@ void FConcertClientPackageManager::DeletePackageFile(const FConcertPackage& Pack
 
 	if (bSuccess)
 	{
-		PackagesPendingPurge.Add(Package.Info.PackageName);
-		PackagesPendingHotReload.Remove(Package.Info.PackageName);
+		PackagesPendingPurge.Add(PackageInfo.PackageName);
+		PackagesPendingHotReload.Remove(PackageInfo.PackageName);
 	}
 }
 
@@ -448,4 +489,14 @@ void FConcertClientPackageManager::PurgePendingPackages()
 	}
 }
 
+bool FConcertClientPackageManager::CanExchangePackageDataAsByteArray(int64 PackageDataSize) const
+{
+	if (FileSharingService && EnumHasAnyFlags(LiveSession->GetSessionFlags(), EConcertSyncSessionFlags::EnableFileSharing))
+	{
+		return FConcertPackage::ShouldEmbedPackageDataAsByteArray(PackageDataSize); // Test the package data size against a preferred limit.
+	}
+
+	return FConcertPackage::CanEmbedPackageDataAsByteArray(PackageDataSize); // The the package data size against the maximum permitted.
+}
+
 #undef LOCTEXT_NAMESPACE
diff --git a/Engine/Plugins/Developer/Concert/ConcertSync/ConcertSyncClient/Source/ConcertSyncClient/Private/ConcertClientPackageManager.h b/Engine/Plugins/Developer/Concert/ConcertSync/ConcertSyncClient/Source/ConcertSyncClient/Private/ConcertClientPackageManager.h
--- a/Engine/Plugins/Developer/Concert/ConcertSync/ConcertSyncClient/Source/ConcertSyncClient/Private/ConcertClientPackageManager.h
+++ b/Engine/Plugins/Developer/Concert/ConcertSync/ConcertSyncClient/Source/ConcertSyncClient/Private/ConcertClientPackageManager.h
@@ -9,19 +9,30 @@
 class FConcertSyncClientLiveSession;
 class FConcertSandboxPlatformFile;
 class ISourceControlProvider;
+class IConcertFileSharingService;
 class UPackage;
 
 struct FAssetData;
 struct FConcertPackage;
 struct FConcertPackageInfo;
 struct FConcertPackageRejectedEvent;
+struct FConcertPackageDataStream;
 
 enum class EConcertPackageUpdateType : uint8;
 
+/** Invoked when a package is too large to be sent to the server if the transport protocol doesn't support large file. */
+DECLARE_MULTICAST_DELEGATE_ThreeParams(FOnConcertClientPackageTooLargeError, const FConcertPackageInfo& /*PackageInfo*/, int64 /*PkgSize*/, int64 /*MaxSize*/);
+
 class FConcertClientPackageManager
 {
 public:
-	FConcertClientPackageManager(TSharedRef<FConcertSyncClientLiveSession> InLiveSession, IConcertClientPackageBridge* InPackageBridge);
+	/**
+	 * Construct the package manager for the specified session.
+	 * @param InLiveSession The session for which the package manager will be used.
+	 * @param InPackageBridge The package bridge (used by the package manager to register itself for package events)
+	 * @param InFileSharingService Optional service (can be null) used by the client to exchange very large packages with the server.
+	 */
+	FConcertClientPackageManager(TSharedRef<FConcertSyncClientLiveSession> InLiveSession, IConcertClientPackageBridge* InPackageBridge, TSharedPtr<IConcertFileSharingService> InFileSharingService);
 	~FConcertClientPackageManager();
 	
 	/**
@@ -76,11 +87,16 @@ public:
 	 */
 	bool PersistSessionChanges(TArrayView<const FName> InPackagesToPersist, ISourceControlProvider* SourceControlProvider, TArray<FText>* OutFailureReasons = nullptr);
 
+	/**
+	 * Called when a package is too big to be handled by the system.
+	 */
+	FOnConcertClientPackageTooLargeError& OnConcertClientPackageTooLargeError() { return OnPackageTooLargeErrorDelegate; }
+
 private:
 	/**
 	 * Apply the data in the given package to disk and update the in-memory state.
 	 */
-	void ApplyPackageUpdate(const FConcertPackage& InPackage);
+	void ApplyPackageUpdate(const FConcertPackageInfo& InPackageInfo, FConcertPackageDataStream& InPackageDataStream);
 
 	/**
 	 * Handle a rejected package event, those are sent by the server when a package update is refused.
@@ -96,17 +112,17 @@ private:
 	/**
 	 * Called to handle a local package event.
 	 */
-	void HandleLocalPackageEvent(const FConcertPackage& Package);
+	void HandleLocalPackageEvent(const FConcertPackageInfo& PackageInfo, const FString& PackagePathname);
 
 	/**
 	 * Utility to save new package data to disk, and also queue if for hot-reload.
 	 */
-	void SavePackageFile(const FConcertPackage& Package);
+	void SavePackageFile(const FConcertPackageInfo& PackageInfo, FConcertPackageDataStream& InPackageDataStream);
 
 	/**
 	 * Utility to remove existing package data from disk, and also queue if for purging.
 	 */
-	void DeletePackageFile(const FConcertPackage& Package);
+	void DeletePackageFile(const FConcertPackageInfo& PackageInfo);
 
 	/**
 	 * Can we currently perform content hot-reloads or purges?
@@ -124,6 +140,11 @@ private:
 	 */
 	void PurgePendingPackages();
 
+	/**
+	 * Whether the package data is small enough to be exchanged using a single TArray<> data structure.
+	 */
+	bool CanExchangePackageDataAsByteArray(int64 PackageDataSize) const;
+
 #if WITH_EDITOR
 	/**
 	 * Sandbox for storing package changes to disk within a Concert session.
@@ -162,4 +183,19 @@ private:
 	 * Array of package names that are pending an in-memory purge.
 	 */
 	TArray<FName> PackagesPendingPurge;
+
+	/**
+	 * Called when a package is too large to be handled by the system.
+	 */
+	FOnConcertClientPackageTooLargeError OnPackageTooLargeErrorDelegate;
+
+	/**
+	 * Optional side channel to exchange large blobs (package data) with the server in a scalable way (ex. the request/response transport layer is not designed and doesn't support exchanging 3GB packages).
+	 */
+	TSharedPtr<IConcertFileSharingService> FileSharingService;
+
+	/**
+	 * Keep the list of package for which the pristine state was sent. Used if the session doesn't support live sync. 
+	 */
+	TSet<FName> EmittedPristinePackages;
 };
diff --git a/Engine/Plugins/Developer/Concert/ConcertSync/ConcertSyncClient/Source/ConcertSyncClient/Private/ConcertClientWorkspace.cpp b/Engine/Plugins/Developer/Concert/ConcertSync/ConcertSyncClient/Source/ConcertSyncClient/Private/ConcertClientWorkspace.cpp
--- a/Engine/Plugins/Developer/Concert/ConcertSync/ConcertSyncClient/Source/ConcertSyncClient/Private/ConcertClientWorkspace.cpp
+++ b/Engine/Plugins/Developer/Concert/ConcertSync/ConcertSyncClient/Source/ConcertSyncClient/Private/ConcertClientWorkspace.cpp
@@ -8,6 +8,7 @@
 #include "IConcertClient.h"
 #include "IConcertModule.h"
 #include "IConcertSession.h"
+#include "IConcertFileSharingService.h"
 #include "ConcertSyncClientLiveSession.h"
 #include "ConcertSyncSessionDatabase.h"
 #include "ConcertSyncSettings.h"
@@ -111,7 +112,8 @@ private:
 };
 
 
-FConcertClientWorkspace::FConcertClientWorkspace(TSharedRef<FConcertSyncClientLiveSession> InLiveSession, IConcertClientPackageBridge* InPackageBridge, IConcertClientTransactionBridge* InTransactionBridge)
+FConcertClientWorkspace::FConcertClientWorkspace(TSharedRef<FConcertSyncClientLiveSession> InLiveSession, IConcertClientPackageBridge* InPackageBridge, IConcertClientTransactionBridge* InTransactionBridge, TSharedPtr<IConcertFileSharingService> InFileSharingService)
+	: FileSharingService(MoveTemp(InFileSharingService))
 {
 	BindSession(InLiveSession, InPackageBridge, InTransactionBridge);
 }
@@ -280,60 +282,9 @@ TFuture<TOptional<FConcertSyncTransactionEvent>> FConcertClientWorkspace::FindOr
 	return MakeFulfilledPromise<TOptional<FConcertSyncTransactionEvent>>().GetFuture(); // Not found.
 }
 
-bool FConcertClientWorkspace::FindPackageEvent(const int64 PackageEventId, FConcertSyncPackageEvent& OutPackageEvent, const bool bMetaDataOnly) const
+bool FConcertClientWorkspace::FindPackageEvent(const int64 PackageEventId, FConcertSyncPackageEventMetaData& OutPackageEvent) const
 {
-	bool bFound = LiveSession->GetSessionDatabase().GetPackageEvent(PackageEventId, OutPackageEvent, bMetaDataOnly);
-	return bFound && (bMetaDataOnly || !IsPackageEventPartiallySynced(OutPackageEvent)); // Avoid succeeding if the event is partially sync but full event data was requested.
-}
-
-TFuture<TOptional<FConcertSyncPackageEvent>> FConcertClientWorkspace::FindOrRequestPackageEvent(const int64 PackageEventId, const bool bMetaDataOnly)
-{
-	// Check if the event exist in the database.
-	FConcertSyncPackageEvent PackageEvent;
-	if (LiveSession->GetSessionDatabase().GetPackageEvent(PackageEventId, PackageEvent, bMetaDataOnly))
-	{
-		// If the package data is required and the event has only the meta data (partially synced, the event was superseded by another).
-		if (!bMetaDataOnly && IsPackageEventPartiallySynced(PackageEvent))
-		{
-			FConcertSyncEventRequest SyncEventRequest{EConcertSyncActivityEventType::Package, PackageEventId };
-			TWeakPtr<FConcertSyncClientLiveSession> WeakLiveSession = LiveSession;
-			return LiveSession->GetSession().SendCustomRequest<FConcertSyncEventRequest, FConcertSyncEventResponse>(SyncEventRequest, LiveSession->GetSession().GetSessionServerEndpointId()).Next([WeakLiveSession, PackageEventId](const FConcertSyncEventResponse& Response)
-			{
-				if (Response.Event.UncompressedPayloadSize > 0) // Some data was sent back?
-				{
-					// Extract the payload as FConcertSyncPackageEvent.
-					FStructOnScope EventPayload;
-					Response.Event.GetPayload(EventPayload);
-					check(EventPayload.IsValid() && EventPayload.GetStruct()->IsChildOf(FConcertSyncPackageEvent::StaticStruct()));
-					FConcertSyncPackageEvent* PackageEvent = (FConcertSyncPackageEvent*)EventPayload.GetStructMemory();
-
-					// Update the database, caching the event to avoid syncing again.
-					if (TSharedPtr<FConcertSyncClientLiveSession> LiveSessionPin = WeakLiveSession.Pin())
-					{
-						LiveSessionPin->GetSessionDatabase().UpdatePackageEvent(PackageEventId, *PackageEvent);
-						// NOTE: PostActivityUpdated() could be called, but the activity did not change, more info was simply cached locally. Unless a use case requires it, don't call it.
-					}
-
-					return TOptional<FConcertSyncPackageEvent>(MoveTemp(*PackageEvent));
-				}
-				else
-				{
-					return TOptional<FConcertSyncPackageEvent>(); // The server did not return any data.
-				}
-			});
-		}
-		else
-		{
-			return MakeFulfilledPromise<TOptional<FConcertSyncPackageEvent>>(MoveTemp(PackageEvent)).GetFuture(); // All required data was already available locally.
-		}
-	}
-
-	return MakeFulfilledPromise<TOptional<FConcertSyncPackageEvent>>().GetFuture(); // Not found.
-}
-
-bool FConcertClientWorkspace::IsPackageEventPartiallySynced(const FConcertSyncPackageEvent& PackageEvent) const
-{
-	return PackageEvent.Package.PackageData.Num() == 0;
+	return LiveSession->GetSessionDatabase().GetPackageEventMetaData(PackageEventId, OutPackageEvent.PackageRevision, OutPackageEvent.PackageInfo);
 }
 
 bool FConcertClientWorkspace::IsTransactionEventPartiallySynced(const FConcertSyncTransactionEvent& TransactionEvent) const
@@ -415,7 +366,7 @@ void FConcertClientWorkspace::BindSession(TSharedPtr<FConcertSyncClientLiveSessi
 	// Create Package Manager
 	if (EnumHasAnyFlags(LiveSession->GetSessionFlags(), EConcertSyncSessionFlags::EnablePackages))
 	{
-		PackageManager = MakeUnique<FConcertClientPackageManager>(LiveSession.ToSharedRef(), InPackageBridge);
+		PackageManager = MakeUnique<FConcertClientPackageManager>(LiveSession.ToSharedRef(), InPackageBridge, FileSharingService);
 	}
 
 	// Create Lock Manager
@@ -930,28 +881,53 @@ void FConcertClientWorkspace::SetTransactionActivity(const FConcertSyncTransacti
 
 void FConcertClientWorkspace::SetPackageActivity(const FConcertSyncPackageActivity& InPackageActivity)
 {
-	// Update this activity
-	if (LiveSession->GetSessionDatabase().SetPackageActivity(InPackageActivity))
+	const FConcertSyncActivity& PackageActivityBasePart = InPackageActivity; // The base activity part correspond the the base class.
+	FConcertSyncPackageEventData PackageActivityEventPart; // The event part is going to be streamed in the database (because it can be huge)
+	PackageActivityEventPart.MetaData.PackageRevision = InPackageActivity.EventData.PackageRevision;
+	PackageActivityEventPart.MetaData.PackageInfo = InPackageActivity.EventData.Package.Info;
+
+	auto SetPackageActivityFn = [&](const FConcertSyncActivity& BaseActivityPart, FConcertSyncPackageEventData& EventDataPart)
 	{
-		PostActivityUpdated(InPackageActivity);
-		if (PackageManager)
+		// Update this activity
+		if (LiveSession->GetSessionDatabase().SetPackageActivity(BaseActivityPart, EventDataPart, /*bMetaDataOnly*/false))
 		{
-			PackageManager->HandleRemotePackage(InPackageActivity.EndpointId, InPackageActivity.EventId, bHasSyncedWorkspace);
+			PostActivityUpdated(BaseActivityPart);
+			if (PackageManager)
+			{
+				PackageManager->HandleRemotePackage(InPackageActivity.EndpointId, InPackageActivity.EventId, bHasSyncedWorkspace);
+			}
+			if (LiveTransactionAuthors)
+			{
+				LiveTransactionAuthors->ResolveLiveTransactionAuthorsForPackage(InPackageActivity.EventData.Package.Info.PackageName);
+			}
 		}
-		if (LiveTransactionAuthors)
+		else
 		{
-			LiveTransactionAuthors->ResolveLiveTransactionAuthorsForPackage(InPackageActivity.EventData.Package.Info.PackageName);
+			UE_LOG(LogConcert, Error, TEXT("Failed to set package activity '%s' on live session '%s': %s"), *LexToString(InPackageActivity.ActivityId), *LiveSession->GetSession().GetName(), *LiveSession->GetSessionDatabase().GetLastError());
 		}
+	};
+
+	// Decide how the package data (part of the package event) is going to be streamed in the database.
+	if (!InPackageActivity.EventData.Package.FileId.IsEmpty())
+	{
+		TSharedPtr<FArchive> PackageDataAr = FileSharingService->CreateReader(InPackageActivity.EventData.Package.FileId); // Package data is available as a shared file.
+		PackageActivityEventPart.PackageDataStream.DataAr = PackageDataAr.Get();
+		PackageActivityEventPart.PackageDataStream.DataSize = PackageDataAr ? PackageDataAr->TotalSize() : 0;
+		SetPackageActivityFn(PackageActivityBasePart, PackageActivityEventPart);
 	}
 	else
 	{
-		UE_LOG(LogConcert, Error, TEXT("Failed to set package activity '%s' on live session '%s': %s"), *LexToString(InPackageActivity.ActivityId), *LiveSession->GetSession().GetName(), *LiveSession->GetSessionDatabase().GetLastError());
+		FMemoryReader Ar(InPackageActivity.EventData.Package.PackageData); // Package data is embedded in the activity itself.
+		PackageActivityEventPart.PackageDataStream.DataAr = &Ar;
+		PackageActivityEventPart.PackageDataStream.DataSize = Ar.TotalSize();
+		PackageActivityEventPart.PackageDataStream.DataBlob = &InPackageActivity.EventData.Package.PackageData;
+		SetPackageActivityFn(PackageActivityBasePart, PackageActivityEventPart);
 	}
 }
 
 void FConcertClientWorkspace::PostActivityUpdated(const FConcertSyncActivity& InActivity)
 {
-	FConcertSyncPackageActivity Activity;
+	FConcertSyncActivity Activity;
 	if (LiveSession->GetSessionDatabase().GetActivity(InActivity.ActivityId, Activity))
 	{
 		FConcertSyncEndpointData EndpointData;
diff --git a/Engine/Plugins/Developer/Concert/ConcertSync/ConcertSyncClient/Source/ConcertSyncClient/Private/ConcertClientWorkspace.h b/Engine/Plugins/Developer/Concert/ConcertSync/ConcertSyncClient/Source/ConcertSyncClient/Private/ConcertClientWorkspace.h
--- a/Engine/Plugins/Developer/Concert/ConcertSync/ConcertSyncClient/Source/ConcertSyncClient/Private/ConcertClientWorkspace.h
+++ b/Engine/Plugins/Developer/Concert/ConcertSync/ConcertSyncClient/Source/ConcertSyncClient/Private/ConcertClientWorkspace.h
@@ -12,6 +12,7 @@
 
 class IConcertClientSession;
 class IConcertClientPackageBridge;
+class IConcertFileSharingService;
 class FConcertClientPackageManager;
 class IConcertClientTransactionBridge;
 class FConcertClientTransactionManager;
@@ -26,7 +27,7 @@ struct FScopedSlowTask;
 class FConcertClientWorkspace : public IConcertClientWorkspace
 {
 public:
-	FConcertClientWorkspace(TSharedRef<FConcertSyncClientLiveSession> InLiveSession, IConcertClientPackageBridge* InPackageBridge, IConcertClientTransactionBridge* InTransactionBridge);
+	FConcertClientWorkspace(TSharedRef<FConcertSyncClientLiveSession> InLiveSession, IConcertClientPackageBridge* InPackageBridge, IConcertClientTransactionBridge* InTransactionBridge, TSharedPtr<IConcertFileSharingService> InFileSharingService);
 	virtual ~FConcertClientWorkspace();
 
 	// IConcertClientWorkspace interface
@@ -43,8 +44,7 @@ public:
 	virtual bool ShouldIgnorePackageDirtyEvent(class UPackage* InPackage) const override;
 	virtual bool FindTransactionEvent(const int64 TransactionEventId, FConcertSyncTransactionEvent& OutTransactionEvent, const bool bMetaDataOnly) const override;
 	virtual TFuture<TOptional<FConcertSyncTransactionEvent>> FindOrRequestTransactionEvent(const int64 TransactionEventId, const bool bMetaDataOnly) override;
-	virtual bool FindPackageEvent(const int64 PackageEventId, FConcertSyncPackageEvent& OutPackageEvent, const bool bMetaDataOnly) const override;
-	virtual TFuture<TOptional<FConcertSyncPackageEvent>> FindOrRequestPackageEvent(const int64 PackageEventId, const bool bMetaDataOnly) override;
+	virtual bool FindPackageEvent(const int64 PackageEventId, FConcertSyncPackageEventMetaData& OutPackageEvent) const override;
 	virtual void GetActivities(const int64 FirstActivityIdToFetch, const int64 MaxNumActivities, TMap<FGuid, FConcertClientInfo>& OutEndpointClientInfoMap, TArray<FConcertClientSessionActivity>& OutActivities) const override;
 	virtual int64 GetLastActivityId() const override;
 	virtual FOnActivityAddedOrUpdated& OnActivityAddedOrUpdated() override;
@@ -162,13 +162,6 @@ private:
 	 */
 	void PostActivityUpdated(const FConcertSyncActivity& InActivity);
 
-	/**
-	 * Check whether a package activity is partially synced, i.e. that only the meta data
-	 * was synced because the activity event was superseded by another one and the package data
-	 * wasn't required to reconstruct the state of a level.
-	 */
-	bool IsPackageEventPartiallySynced(const FConcertSyncPackageEvent& PackageEvent) const;
-
 	/**
 	 * Check whether a transaction activity is partially synced, i.e. that only the meta data
 	 * was synced because the activity event was superseded by another one and the transaction data
@@ -219,4 +212,7 @@ private:
 
 	/** True if the client has marked the further transaction as 'non-ignored'. This is sent at the end of the frame. */
 	bool bPendingStopIgnoringActivityOnRestore = false;
+
+	/** Optional side channel to exchange large blobs (package data) with the server in a scalable way (ex. the request/response transport layer is not designed and doesn't support exchanging 3GB packages). */
+	TSharedPtr<IConcertFileSharingService> FileSharingService;
 };
diff --git a/Engine/Plugins/Developer/Concert/ConcertSync/ConcertSyncClient/Source/ConcertSyncClient/Private/ConcertSyncClient.cpp b/Engine/Plugins/Developer/Concert/ConcertSync/ConcertSyncClient/Source/ConcertSyncClient/Private/ConcertSyncClient.cpp
--- a/Engine/Plugins/Developer/Concert/ConcertSync/ConcertSyncClient/Source/ConcertSyncClient/Private/ConcertSyncClient.cpp
+++ b/Engine/Plugins/Developer/Concert/ConcertSync/ConcertSyncClient/Source/ConcertSyncClient/Private/ConcertSyncClient.cpp
@@ -115,7 +115,6 @@ void FConcertSyncClient::PersistAllSessionChanges()
 #endif
 }
 
-
 void FConcertSyncClient::GetSessionClientActions(const FConcertSessionClientInfo& InClientInfo, TArray<FConcertActionDefinition>& OutActions) const
 {
 #if WITH_EDITOR
@@ -126,10 +125,16 @@ void FConcertSyncClient::GetSessionClientActions(const FConcertSessionClientInfo
 #endif
 }
 
+void FConcertSyncClient::SetFileSharingService(TSharedPtr<IConcertFileSharingService> InFileSharingService)
+{
+	check(!FileSharingService); // Not really meant to be set more than once.
+	FileSharingService = MoveTemp(InFileSharingService);
+}
+
 void FConcertSyncClient::CreateWorkspace(const TSharedRef<FConcertSyncClientLiveSession>& InLiveSession)
 {
 	DestroyWorkspace();
-	Workspace = MakeShared<FConcertClientWorkspace>(InLiveSession, PackageBridge, TransactionBridge);
+	Workspace = MakeShared<FConcertClientWorkspace>(InLiveSession, PackageBridge, TransactionBridge, FileSharingService);
 	OnWorkspaceStartupDelegate.Broadcast(Workspace);
 #if WITH_EDITOR
 	if (GIsEditor && EnumHasAllFlags(SessionFlags, EConcertSyncSessionFlags::EnablePackages | EConcertSyncSessionFlags::ShouldUsePackageSandbox))
diff --git a/Engine/Plugins/Developer/Concert/ConcertSync/ConcertSyncClient/Source/ConcertSyncClient/Private/ConcertSyncClient.h b/Engine/Plugins/Developer/Concert/ConcertSync/ConcertSyncClient/Source/ConcertSyncClient/Private/ConcertSyncClient.h
--- a/Engine/Plugins/Developer/Concert/ConcertSync/ConcertSyncClient/Source/ConcertSyncClient/Private/ConcertSyncClient.h
+++ b/Engine/Plugins/Developer/Concert/ConcertSync/ConcertSyncClient/Source/ConcertSyncClient/Private/ConcertSyncClient.h
@@ -39,6 +39,8 @@ public:
 	virtual void PersistAllSessionChanges() override;
 	virtual void GetSessionClientActions(const FConcertSessionClientInfo& InClientInfo, TArray<FConcertActionDefinition>& OutActions) const override;
 
+	virtual void SetFileSharingService(TSharedPtr<IConcertFileSharingService> InFileSharingService) override;
+
 private:
 	void CreateWorkspace(const TSharedRef<FConcertSyncClientLiveSession>& InLiveSession);
 	void DestroyWorkspace();
@@ -76,6 +78,9 @@ private:
 	/** Delegate called just after the workspace was deleted. */
 	FOnConcertClientSyncSessionStartupOrShutdown OnSyncSessionShutdownDelegate;
 
+	/** Optional side channel to exchange large blobs (package data) with the server in a scalable way (ex. the request/response transport layer is not designed and doesn't support exchanging 3GB packages). */
+	TSharedPtr<IConcertFileSharingService> FileSharingService;
+
 #if WITH_EDITOR
 	/** Presence manager for the current session. */
 	TSharedPtr<FConcertClientPresenceManager> PresenceManager;
diff --git a/Engine/Plugins/Developer/Concert/ConcertSync/ConcertSyncClient/Source/ConcertSyncClient/Public/IConcertClientPackageBridge.h b/Engine/Plugins/Developer/Concert/ConcertSync/ConcertSyncClient/Source/ConcertSyncClient/Public/IConcertClientPackageBridge.h
--- a/Engine/Plugins/Developer/Concert/ConcertSync/ConcertSyncClient/Source/ConcertSyncClient/Public/IConcertClientPackageBridge.h
+++ b/Engine/Plugins/Developer/Concert/ConcertSync/ConcertSyncClient/Source/ConcertSyncClient/Public/IConcertClientPackageBridge.h
@@ -6,9 +6,9 @@
 #include "Delegates/Delegate.h"
 
 class UPackage;
-struct FConcertPackage;
+struct FConcertPackageInfo;
 
-DECLARE_MULTICAST_DELEGATE_OneParam(FOnConcertClientLocalPackageEvent, const FConcertPackage&);
+DECLARE_MULTICAST_DELEGATE_TwoParams(FOnConcertClientLocalPackageEvent, const FConcertPackageInfo&, const FString&);
 DECLARE_MULTICAST_DELEGATE_OneParam(FOnConcertClientLocalPackageDiscarded, UPackage*);
 
 /**
diff --git a/Engine/Plugins/Developer/Concert/ConcertSync/ConcertSyncClient/Source/ConcertSyncClient/Public/IConcertClientWorkspace.h b/Engine/Plugins/Developer/Concert/ConcertSync/ConcertSyncClient/Source/ConcertSyncClient/Public/IConcertClientWorkspace.h
--- a/Engine/Plugins/Developer/Concert/ConcertSync/ConcertSyncClient/Source/ConcertSyncClient/Public/IConcertClientWorkspace.h
+++ b/Engine/Plugins/Developer/Concert/ConcertSync/ConcertSyncClient/Source/ConcertSyncClient/Public/IConcertClientWorkspace.h
@@ -166,25 +166,10 @@ public:
 	/**
 	 * Lookup the specified package event.
 	 * @param[in] PackageEventId ID of the package to look for.
-	 * @param[out] OutPackageEvent Information about the package.
-	 * @param[in] bMetaDataOnly True to extract the event meta-data only (title, ids, etc), false if the package data is also required.
-	 * @return Whether the package event was found and requested data available. If bMetaDataOnly is false and the event was partially
-	 *         synced (because it was superseded by another one), the function returns false.
-	 * @see FindOrRequestPackageEvent() if the package data (not only the activity summary) is required.
-	 * @note This function is more efficient to use if only the meta data is required.
-	 */
-	virtual bool FindPackageEvent(const int64 PackageEventId, FConcertSyncPackageEvent& OutPackageEvent, const bool bMetaDataOnly) const = 0;
-
-	/**
-	 * Lookup the specified package events. By default, some package activities are partially synced (only the summary) when the system detects that
-	 * the package data was superseded by another version.
-	 * @param[in] PackageEventId ID of the package to look for.
-	 * @param[in] bMetaDataOnly True to extract the event meta-data only (title, ids, etc), false if the package data is also required.
-	 * @return A future containing the event, if found. If bMetaDataOnly is false and the event was partially synced (because it was superseded
-	 *         by another one) the function will perform an asynchronous server request to get the package data.
-	 * @note If only the meta data is required, consider using FindPackageEvent() as it is more efficient.
+	 * @param[out] OutPackageEvent Information about the package except the package data itself.
+	 * @return Whether the package event meta data was found.
 	 */
-	virtual TFuture<TOptional<FConcertSyncPackageEvent>> FindOrRequestPackageEvent(const int64 PackageEventId, const bool bMetaDataOnly) = 0;
+	virtual bool FindPackageEvent(const int64 PackageEventId, FConcertSyncPackageEventMetaData& OutPackageEvent) const = 0;
 
 	/**
 	 * @return the delegate called every time the workspace is synced.
diff --git a/Engine/Plugins/Developer/Concert/ConcertSync/ConcertSyncClient/Source/ConcertSyncClient/Public/IConcertSyncClient.h b/Engine/Plugins/Developer/Concert/ConcertSync/ConcertSyncClient/Source/ConcertSyncClient/Public/IConcertSyncClient.h
--- a/Engine/Plugins/Developer/Concert/ConcertSync/ConcertSyncClient/Source/ConcertSyncClient/Public/IConcertSyncClient.h
+++ b/Engine/Plugins/Developer/Concert/ConcertSync/ConcertSyncClient/Source/ConcertSyncClient/Public/IConcertSyncClient.h
@@ -14,6 +14,7 @@ class IConcertClientSequencerManager;
 struct FConcertSessionClientInfo;
 
 class IConcertSyncClient;
+class IConcertFileSharingService;
 
 DECLARE_MULTICAST_DELEGATE_OneParam(FOnConcertClientWorkspaceStartupOrShutdown, const TSharedPtr<IConcertClientWorkspace>& /** InClientWorkspace */ );
 DECLARE_MULTICAST_DELEGATE_OneParam(FOnConcertClientSyncSessionStartupOrShutdown, const IConcertSyncClient* /* InSyncClient */ );
@@ -85,4 +86,7 @@ public:
 	 * @param OutActions The available actions, if any.
 	 */
 	virtual void GetSessionClientActions(const FConcertSessionClientInfo& InClientInfo, TArray<FConcertActionDefinition>& OutActions) const = 0;
+
+	/** Set the file sharing service, enabling the server to work with large files. The client and the server must use compatible sharing services. The service must be set prior establishing a connection. */
+	virtual void SetFileSharingService(TSharedPtr<IConcertFileSharingService> InFileSharingService) = 0;
 };
diff --git a/Engine/Plugins/Developer/Concert/ConcertSync/ConcertSyncClient/Source/ConcertSyncClient/Test/ConcertClientLiveTransactionAuthorsTests.cpp b/Engine/Plugins/Developer/Concert/ConcertSync/ConcertSyncClient/Source/ConcertSyncClient/Test/ConcertClientLiveTransactionAuthorsTests.cpp
--- a/Engine/Plugins/Developer/Concert/ConcertSync/ConcertSyncClient/Source/ConcertSyncClient/Test/ConcertClientLiveTransactionAuthorsTests.cpp
+++ b/Engine/Plugins/Developer/Concert/ConcertSync/ConcertSyncClient/Source/ConcertSyncClient/Test/ConcertClientLiveTransactionAuthorsTests.cpp
@@ -116,24 +116,26 @@ void FenceTransactions(FConcertSyncClientLiveSession& InLiveSession, FConcertCli
 {
 	InLiveSession.GetSessionDatabase().SetEndpoint(InSourceClient.EndpointId, InSourceClient.EndpointData);
 
-	FConcertSyncPackageActivity PackageActivity;
-	PackageActivity.EndpointId = InSourceClient.EndpointId;
-	PackageActivity.EventData.Package.Info.PackageName = InPackageName;
-	PackageActivity.EventData.Package.Info.PackageUpdateType = EConcertPackageUpdateType::Dummy;
+	FConcertSyncActivity PackageActivityBasePart;
+	FConcertPackageInfo PackageInfo;
+	FConcertPackageDataStream PackageDataStream;
+	PackageActivityBasePart.EndpointId = InSourceClient.EndpointId;
+	PackageInfo.PackageName = InPackageName;
+	PackageInfo.PackageUpdateType = EConcertPackageUpdateType::Dummy;
 	if (InTransactionIdOverride)
 	{
-		PackageActivity.EventData.Package.Info.TransactionEventIdAtSave = *InTransactionIdOverride;
+		PackageInfo.TransactionEventIdAtSave = *InTransactionIdOverride;
 	}
 	else
 	{
-		InLiveSession.GetSessionDatabase().GetTransactionMaxEventId(PackageActivity.EventData.Package.Info.TransactionEventIdAtSave);
+		InLiveSession.GetSessionDatabase().GetTransactionMaxEventId(PackageInfo.TransactionEventIdAtSave);
 	}
 
 	int64 ActivityId = 0;
 	int64 EventId = 0;
-	InLiveSession.GetSessionDatabase().AddPackageActivity(PackageActivity, ActivityId, EventId);
+	InLiveSession.GetSessionDatabase().AddPackageActivity(PackageActivityBasePart, PackageInfo, PackageDataStream, ActivityId, EventId);
 
-	InLiveTransactionAuthors.ResolveLiveTransactionAuthorsForPackage(PackageActivity.EventData.Package.Info.PackageName);
+	InLiveTransactionAuthors.ResolveLiveTransactionAuthorsForPackage(PackageInfo.PackageName);
 }
 
 /** Ensures the live transaction authors works correctly when there is no other clients connected. */
diff --git a/Engine/Plugins/Developer/Concert/ConcertSync/ConcertSyncClient/Source/ConcertSyncUI/Private/SConcertSessionActivities.cpp b/Engine/Plugins/Developer/Concert/ConcertSync/ConcertSyncClient/Source/ConcertSyncUI/Private/SConcertSessionActivities.cpp
--- a/Engine/Plugins/Developer/Concert/ConcertSync/ConcertSyncClient/Source/ConcertSyncUI/Private/SConcertSessionActivities.cpp
+++ b/Engine/Plugins/Developer/Concert/ConcertSync/ConcertSyncClient/Source/ConcertSyncUI/Private/SConcertSessionActivities.cpp
@@ -64,13 +64,22 @@ FText GetOperationName(const FConcertClientSessionActivity& Activity)
 
 	if (const FConcertSyncPackageActivitySummary* PackageSummary = Activity.ActivitySummary.Cast<FConcertSyncPackageActivitySummary>())
 	{
+		auto GetSavedOperationNameFn = [](const FConcertSyncPackageActivitySummary* PackageSummary)
+		{
+			if (PackageSummary->bAutoSave)
+			{
+				return LOCTEXT("AutoSavePackageOperation", "Auto-Save Package");
+			}
+			return PackageSummary->bPreSave ? LOCTEXT("PreSavePackageOperation", "Pre-Save Package") : LOCTEXT("SavePackageOperation", "Save Package");
+		};
+
 		switch (PackageSummary->PackageUpdateType)
 		{
 		case EConcertPackageUpdateType::Added   : return LOCTEXT("NewPackageOperation",    "New Package");
 		case EConcertPackageUpdateType::Deleted : return LOCTEXT("DeletePackageOperation", "Delete Package");
 		case EConcertPackageUpdateType::Renamed : return LOCTEXT("RenamePackageOperation", "Rename Package");
-		case EConcertPackageUpdateType::Saved   : return LOCTEXT("SavePackageOperation",   "Save Package");
-		case EConcertPackageUpdateType::Dummy: // Fall-Through
+		case EConcertPackageUpdateType::Saved   : return GetSavedOperationNameFn(PackageSummary);
+		case EConcertPackageUpdateType::Dummy   : return LOCTEXT("DiscardPackageOperation", "Discard Changes");
 		default: break;
 		}
 	}
@@ -613,7 +622,8 @@ void SConcertSessionActivities::UpdateDetailArea(TSharedPtr<FConcertClientSessio
 		{
 			FConcertSyncPackageEvent PackageEvent;
 			InSelectedActivity->EventPayload->GetTypedPayload(PackageEvent);
-			DisplayPackageDetails(*InSelectedActivity, PackageEvent);
+			checkf(!PackageEvent.Package.HasPackageData(), TEXT("UI should only request the package meta data because the package data is not useful and can be very large"));
+			DisplayPackageDetails(*InSelectedActivity, PackageEvent.PackageRevision, PackageEvent.Package.Info);
 		}
 		else // Other activity types (lock/connection) don't have details panel.
 		{
@@ -647,26 +657,15 @@ void SConcertSessionActivities::UpdateDetailArea(TSharedPtr<FConcertClientSessio
 	else if (InSelectedActivity->Activity.EventType == EConcertSyncActivityEventType::Package && GetPackageEventFn) // A function is bound to get the package event?
 	{
 		SetDetailsPanelVisibility(LoadingDetailsPanel.Get());
-		TWeakPtr<SConcertSessionActivities> WeakSelf = SharedThis(this);
-		GetPackageEventFn(*InSelectedActivity).Next([WeakSelf, InSelectedActivity](const TOptional<FConcertSyncPackageEvent>& PackageEvent)
+		FConcertSyncPackageEventMetaData PackageEventMetaData;
+		if (GetPackageEventFn(*InSelectedActivity, PackageEventMetaData))
 		{
-			if (TSharedPtr<SConcertSessionActivities> Self = WeakSelf.Pin()) // If 'this' object hasn't been deleted.
-			{
-				if (Self->GetSelectedActivity() == InSelectedActivity) // Ensure the activity is still selected.
-				{
-					if (PackageEvent.IsSet())
-					{
-						Self->DisplayPackageDetails(*InSelectedActivity, PackageEvent.GetValue());
-					}
-					else
-					{
-						Self->SetDetailsPanelVisibility(Self->NoDetailsPanel.Get());
-					}
-				}
-				// else -> The details panel is presenting information for another activity (or no activity).
-			}
-			// else -> The widget was deleted.
-		});
+			DisplayPackageDetails(*InSelectedActivity, PackageEventMetaData.PackageRevision, PackageEventMetaData.PackageInfo);
+		}
+		else
+		{
+			SetDetailsPanelVisibility(NoDetailsPanel.Get());
+		}
 	}
 	else
 	{
@@ -952,7 +951,7 @@ void SConcertSessionActivities::DisplayTransactionDetails(const FConcertClientSe
 	SetDetailsPanelVisibility(TransactionDetailsPanel.Get());
 }
 
-void SConcertSessionActivities::DisplayPackageDetails(const FConcertClientSessionActivity& Activity, const FConcertSyncPackageEvent& PackageEvent)
+void SConcertSessionActivities::DisplayPackageDetails(const FConcertClientSessionActivity& Activity, int64 PackageRevision, const FConcertPackageInfo& PackageInfo)
 {
 	const FConcertClientInfo* ClientInfo = nullptr;
 	if (GetActivityUserFn)
@@ -960,7 +959,7 @@ void SConcertSessionActivities::DisplayPackageDetails(const FConcertClientSessio
 		ClientInfo = GetActivityUserFn(Activity.Activity.EndpointId);
 	}
 
-	PackageDetailsPanel->SetPackageInfo(PackageEvent.Package.Info, PackageEvent.PackageRevision, ClientInfo ? ClientInfo->DisplayName : FString());
+	PackageDetailsPanel->SetPackageInfo(PackageInfo, PackageRevision, ClientInfo ? ClientInfo->DisplayName : FString());
 	SetDetailsPanelVisibility(PackageDetailsPanel.Get());
 }
 
diff --git a/Engine/Plugins/Developer/Concert/ConcertSync/ConcertSyncClient/Source/ConcertSyncUI/Public/SConcertSessionActivities.h b/Engine/Plugins/Developer/Concert/ConcertSync/ConcertSyncClient/Source/ConcertSyncUI/Public/SConcertSessionActivities.h
--- a/Engine/Plugins/Developer/Concert/ConcertSync/ConcertSyncClient/Source/ConcertSyncUI/Public/SConcertSessionActivities.h
+++ b/Engine/Plugins/Developer/Concert/ConcertSync/ConcertSyncClient/Source/ConcertSyncUI/Public/SConcertSessionActivities.h
@@ -53,8 +53,8 @@ public:
 	/** Returns the transaction event corresponding the specified activity.*/
 	using FGetTransactionEvent = TFunction<TFuture<TOptional<FConcertSyncTransactionEvent>>(const FConcertClientSessionActivity& /*Activity*/)>;
 
-	/** Returns the package vent corresponding to the package activity. */
-	using FGetPackageEvent = TFunction<TFuture<TOptional<FConcertSyncPackageEvent>>(const FConcertClientSessionActivity& /*Activity*/)>;
+	/** Returns the package event corresponding to the package activity. */
+	using FGetPackageEvent = TFunction<bool(const FConcertClientSessionActivity& /*Activity*/, FConcertSyncPackageEventMetaData& /*OutEvent*/)>;
 
 	/** Used to overlay a widget over a column widget to add custom functionalities to a row. */
 	using FMakeColumnOverlayWidgetFunc = TFunction<TSharedPtr<SWidget>(TWeakPtr<FConcertClientSessionActivity> /*ThisRowActivity*/, const FName& /*ColumnId*/)>;
@@ -212,7 +212,7 @@ private:
 	void DisplayTransactionDetails(const FConcertClientSessionActivity& Activity, const FConcertTransactionEventBase& InTransaction);
 
 	/** Displays what changed in a package in the details area when a package activity is selected in the list. */
-	void DisplayPackageDetails(const FConcertClientSessionActivity& Activity, const FConcertSyncPackageEvent& PackageEvent);
+	void DisplayPackageDetails(const FConcertClientSessionActivity& Activity, int64 PackageRevision, const FConcertPackageInfo& PackageInfo);
 
 	/** Update the detail area to display the proper panel and details. */
 	void UpdateDetailArea(TSharedPtr<FConcertClientSessionActivity> InSelectedActivity);
diff --git a/Engine/Plugins/Developer/Concert/ConcertSync/ConcertSyncCore/Source/ConcertSyncCore/ConcertSyncCore.Build.cs b/Engine/Plugins/Developer/Concert/ConcertSync/ConcertSyncCore/Source/ConcertSyncCore/ConcertSyncCore.Build.cs
--- a/Engine/Plugins/Developer/Concert/ConcertSync/ConcertSyncCore/Source/ConcertSyncCore/ConcertSyncCore.Build.cs
+++ b/Engine/Plugins/Developer/Concert/ConcertSync/ConcertSyncCore/Source/ConcertSyncCore/ConcertSyncCore.Build.cs
@@ -19,6 +19,7 @@ namespace UnrealBuildTool.Rules
 				{
 					"Concert",
 					"ConcertTransport",
+					"Json",
 					"SQLiteCore",
 				}
 			);
diff --git a/Engine/Plugins/Developer/Concert/ConcertSync/ConcertSyncCore/Source/ConcertSyncCore/Private/ConcertLocalFileSharingService.cpp b/Engine/Plugins/Developer/Concert/ConcertSync/ConcertSyncCore/Source/ConcertSyncCore/Private/ConcertLocalFileSharingService.cpp
new file mode 100644
--- /dev/null
+++ b/Engine/Plugins/Developer/Concert/ConcertSync/ConcertSyncCore/Source/ConcertSyncCore/Private/ConcertLocalFileSharingService.cpp
@@ -0,0 +1,127 @@
+// Copyright Epic Games, Inc. All Rights Reserved.
+
+#include "ConcertLocalFileSharingService.h"
+#include "HAL/FileManager.h"
+#include "ConcertUtil.h"
+#include "Serialization/JsonReader.h"
+#include "Serialization/JsonWriter.h"
+#include "Serialization/JsonSerializer.h"
+
+FConcertLocalFileSharingService::FConcertLocalFileSharingService(const FString& Role)
+	: SharedRootPathname(FPaths::EngineIntermediateDir() / Role / TEXT("TempFileShare")) // The EngineIntermediateDir is expected to be know of both the client and the server.
+	, SystemMutexName(TEXT("ConcertFileSharingMutex_777221B1ZZ")) // Arbitrary name, but unique enough to avoid collision with other global mutexes from other apps.
+	, ActiveServicesRepositoryPathname(SharedRootPathname / TEXT("ActiveProcesses.json"))
+{
+	FSystemWideCriticalSection ScopedSysWideMutex(SystemMutexName);
+
+	TArray<uint32> ActivesPids;
+	LoadActiveServices(ActivesPids);
+	RemoveDeadProcessesAndFiles(ActivesPids);
+	ActivesPids.Add(FPlatformProcess::GetCurrentProcessId());
+	SaveActiveServices(ActivesPids);
+}
+
+FConcertLocalFileSharingService::~FConcertLocalFileSharingService()
+{
+	FSystemWideCriticalSection ScopedSysWideMutex(SystemMutexName);
+
+	TArray<uint32> ActivePids;
+	LoadActiveServices(ActivePids);
+	ActivePids.Remove(FPlatformProcess::GetCurrentProcessId());
+	RemoveDeadProcessesAndFiles(ActivePids);
+	if (ActivePids.Num()) // If no pids are left, the sharing directory was deleted, don't recreate it to save an empty file.
+	{
+		SaveActiveServices(ActivePids);
+	}
+}
+
+void FConcertLocalFileSharingService::RemoveDeadProcessesAndFiles(TArray<uint32>& InOutPids)
+{
+	// Remove all dead process.
+	InOutPids.RemoveAll([](uint32 Pid)
+	{
+		return !FPlatformProcess::IsApplicationRunning(Pid);
+	});
+
+	// If no processes are left using the shared directory, delete it with its content.
+	if (InOutPids.Num() == 0)
+	{
+		IFileManager::Get().DeleteDirectory(*SharedRootPathname, /*RequireExists*/false, /*Tree*/true);
+	}
+}
+
+bool FConcertLocalFileSharingService::Publish(const FString& Pathname, FString& OutFileId)
+{
+	// Create a unique name for the file.
+	OutFileId = SharedRootPathname / FGuid::NewGuid().ToString();
+	return IFileManager::Get().Copy(*OutFileId, *Pathname) == COPY_OK;
+}
+
+bool FConcertLocalFileSharingService::Publish(FArchive& SrcAr, int64 Size, FString& OutFileId)
+{
+	// Create a unique name for the file.
+	FString SharedFilePathname = SharedRootPathname / FGuid::NewGuid().ToString();
+	TUniquePtr<FArchive> SharedFileWriter(IFileManager::Get().CreateFileWriter(*SharedFilePathname));
+	if (SharedFileWriter)
+	{
+		ConcertUtil::Copy(*SharedFileWriter, SrcAr, Size);
+		OutFileId = FPaths::ConvertRelativePathToFull(SharedFilePathname);
+		return true;
+	}
+	return false;
+}
+
+TSharedPtr<FArchive> FConcertLocalFileSharingService::CreateReader(const FString& InFileId)
+{
+	if (IFileManager::Get().FileExists(*InFileId))
+	{
+		// Delete the files when the archive goes out of scope, so that consumed files don't use too much space.
+		auto Deleter = [InFileId](FArchive* Ar)
+		{
+			delete Ar;
+			IFileManager::Get().Delete(*InFileId);
+		};
+
+		return MakeShareable(IFileManager::Get().CreateFileReader(*InFileId), MoveTemp(Deleter));
+	}
+
+	return nullptr; // File not found.
+}
+
+void FConcertLocalFileSharingService::LoadActiveServices(TArray<uint32>& OutPids)
+{
+	if (IFileManager::Get().FileExists(*ActiveServicesRepositoryPathname))
+	{
+		TUniquePtr<FArchive> Ar(IFileManager::Get().CreateFileReader(*ActiveServicesRepositoryPathname));
+		TSharedRef<TJsonReader<>> JsonReader = TJsonReaderFactory<>::Create(Ar.Get());
+
+		TArray< TSharedPtr<FJsonValue> > Pids;
+		if (FJsonSerializer::Deserialize(JsonReader, Pids))
+		{
+			for (const TSharedPtr<FJsonValue>& PidVal : Pids)
+			{
+				uint32 Pid;
+				if (PidVal->TryGetNumber(Pid))
+				{
+					OutPids.Add(Pid);
+				}
+			}
+		}
+	}
+}
+
+void FConcertLocalFileSharingService::SaveActiveServices(const TArray<uint32>& InPids)
+{
+	if (IFileManager::Get().DirectoryExists(*SharedRootPathname) || IFileManager::Get().MakeDirectory(*SharedRootPathname, /*Tree*/true))
+	{
+		TUniquePtr<FArchive> Ar(IFileManager::Get().CreateFileWriter(*ActiveServicesRepositoryPathname));
+		TSharedRef<TJsonWriter<>> JsonWriter = TJsonWriterFactory<>::Create(Ar.Get());
+
+		JsonWriter->WriteArrayStart();
+		for (uint32 Pid : InPids)
+		{
+			JsonWriter->WriteValue(static_cast<int64>(Pid));
+		}
+		JsonWriter->WriteArrayEnd();
+	}
+}
diff --git a/Engine/Plugins/Developer/Concert/ConcertSync/ConcertSyncCore/Source/ConcertSyncCore/Private/ConcertSyncSessionDatabase.cpp b/Engine/Plugins/Developer/Concert/ConcertSync/ConcertSyncCore/Source/ConcertSyncCore/Private/ConcertSyncSessionDatabase.cpp
--- a/Engine/Plugins/Developer/Concert/ConcertSync/ConcertSyncCore/Source/ConcertSyncCore/Private/ConcertSyncSessionDatabase.cpp
+++ b/Engine/Plugins/Developer/Concert/ConcertSync/ConcertSyncCore/Source/ConcertSyncCore/Private/ConcertSyncSessionDatabase.cpp
@@ -158,10 +158,53 @@ namespace PackageDataUtil
 {
 
 const int32 MinFilesToCache = 10;
-const uint32 DataVersion = 1;
+const uint8 DataVersion = 2;
 const uint64 MaxFileSizeBytesToCache = 200 * 1024 * 1024;
 const FGuid EntryFooter = FGuid(0x2EFC8CDD, 0x748E46C0, 0xA5485769, 0x13A3C354);
 
+/**
+ * Defines how the package data file written. The 8 LSB are reserved for FileFormatVersion which control how 24 MSB are interpreted.
+ */
+enum class EPackageDataFormat : uint8
+{
+	/** The package data (the body) is written in a single compressed block. The value 1 represents the original format. Still used to store small packages.*/
+	PackageDataCompressed = 1 << 0,
+
+	/** The package data (the body) is written uncompressed. Added at 4.25 to support large package files. */
+	PackageDataUncompressed = 1 << 1,
+};
+ENUM_CLASS_FLAGS(EPackageDataFormat)
+
+/** Write the version/format information in the first 4 bytes of the blob, enabling the reader to parse the blob. */
+void WritePackageBlobVersionInfo(FArchive& Ar, EPackageDataFormat PackageDataFormat)
+{
+	check(Ar.IsSaving() && Ar.Tell() == 0); // Expects to write the version/format info at the beginning.
+
+	// Keep the major DataVersion in the 8 LSB, encode package data format in the 24 MSB.
+	constexpr uint32 BitsPerByte = 8;
+	uint32 SerializedDataVersion = static_cast<uint32>(DataVersion) | (static_cast<uint32>(PackageDataFormat) << (sizeof(DataVersion) * BitsPerByte));
+	Ar.SerializeIntPacked(SerializedDataVersion);
+}
+
+/** Reads the first 4 bytes of the archive and detect how the package data was written in the blob. */
+EPackageDataFormat ParsePackageDataFormat(FArchive& Ar)
+{
+	check(Ar.IsLoading() && Ar.Tell() == 0); // Expected to read the version/format info at the beginning.
+
+	uint32 SerializedDataVersion = 0;
+	Ar.SerializeIntPacked(SerializedDataVersion);
+	uint8 ExtractedDataVersion = static_cast<uint8>(SerializedDataVersion & 0xFF); // DataVersion is in the 8 LSB
+	if (ExtractedDataVersion == 1)
+	{
+		// In Version 1, the package was always written compressed into a single block.
+		return EPackageDataFormat::PackageDataCompressed;
+	}
+
+	// In Version 2, the package data format was encoded after the 8 LSB used for 'DataVersion'
+	constexpr uint32 BitsPerByte = 8;
+	return static_cast<EPackageDataFormat>(SerializedDataVersion >> (sizeof(DataVersion) * BitsPerByte));
+}
+
 FString GetDataPath(const FString& InSessionPath)
 {
 	return InSessionPath / TEXT("Packages");
@@ -177,89 +220,190 @@ FString GetDataFilename(const FName InPackageName, const int64 InRevision)
 	return GetDataFilename(InPackageName.ToString(), InRevision);
 }
 
-bool WritePackage(const FConcertPackageInfo& InPackageInfo, const TArray<uint8>& InPackageData, TArray<uint8>& OutSerializedPackageData)
+/** Internally used by WritePackageBlob() - write the serialized package header. */
+void WritePackageHeader(FArchive& Ar, EPackageDataFormat PackageDataFormat)
 {
-	FMemoryWriter Ar(OutSerializedPackageData);
-
-	// Serialize the data version
-	uint32 SerializedDataVersion = DataVersion;
-	Ar.SerializeIntPacked(SerializedDataVersion);
+	// First 4-bytes contains information about the format.
+	WritePackageBlobVersionInfo(Ar, PackageDataFormat);
 
-	// Serialize the info (header)
-	int64 BodyOffsetTell = Ar.Tell();
-	int64 BodyOffset = 0;
-	Ar << BodyOffset;
-	FConcertPackageInfo::StaticStruct()->SerializeItem(Ar, const_cast<FConcertPackageInfo*>(&InPackageInfo), nullptr);
+	// PackageDataOffset is kept to allow 4.25 to read 4.24 packages. The variable-length data stored between version info and package data was never used, so it was stripped in 4.25.
+	int64 PackageDataOffsetPlaceholderPos = Ar.Tell(); // Remember at which offset the package data offset variable is going to be written.
+	int64 PackageDataOffsetPlaceholder = 0;
+	Ar << PackageDataOffsetPlaceholder; // Reserve a space in the header to store the package data offset.
+	int64 PackageDataOffset = Ar.Tell(); // Read the end position of the header.
+	Ar.Seek(PackageDataOffsetPlaceholderPos); // Rewind to update the place holder value.
+	Ar << PackageDataOffset; // Replace the place holder value with the actual value.
+}
 
-	// Serialize the raw data (body)
-	BodyOffset = Ar.Tell();
-	Ar.Seek(BodyOffsetTell);
-	Ar << BodyOffset;
-	Ar.Seek(BodyOffset);
-	uint32 UncompressedPackageSize = InPackageData.Num();
-	Ar.SerializeIntPacked(UncompressedPackageSize);
-	if (UncompressedPackageSize > 0)
-	{
-		Ar.SerializeCompressed((uint8*)InPackageData.GetData(), UncompressedPackageSize, NAME_Zlib);
-	}
+/** Internally used by ExtractPackageData() - read serialized package header. */
+TTuple<EPackageDataFormat, int64> ReadPackageHeader(FArchive& Ar)
+{
+	// Consume the first bytes to read the package version and package data format.
+	EPackageDataFormat PackageDataFormat = ParsePackageDataFormat(Ar);
+	
+	// Deserialize the next 8 bytes containing the offset to the package data (for binary compatibility with 4.24)
+	int64 PackageDataOffset = 0;
+	Ar << PackageDataOffset;
+	return MakeTuple(PackageDataFormat, PackageDataOffset);
+}
 
-	// Serialize the footer so we know we didn't crash mid-write
+/** Internally used by WritePackageBlob() - write the serialized package footer. */
+void WritePackageFooter(FArchive& Ar)
+{
 	FGuid SerializedFooter = EntryFooter;
 	Ar << SerializedFooter;
+}
 
-	return !Ar.IsError();
+/** Internally used by ExtractPackageData() - ensure the file was fully written. */
+bool HasValidPackageFooter(FArchive& Ar)
+{
+	// Test the footer is in place so we know we didn't crash mid-write
+	const int64 SerializedPackageSize = Ar.TotalSize();
+	if (SerializedPackageSize >= sizeof(FGuid))
+	{
+		FGuid SerializedFooter;
+		Ar.Seek(SerializedPackageSize - sizeof(FGuid));
+		Ar << SerializedFooter;
+		Ar.Seek(0);
+		return SerializedFooter == EntryFooter;
+	}
+	return false;
 }
 
-bool ReadPackage(const TArray<uint8>& InSerializedPackageData, FConcertPackageInfo* OutPackageInfo, TArray<uint8>* OutPackageData)
+int64 EstimatePackageBlobMetaDataSize()
 {
-	FMemoryReader Ar(InSerializedPackageData);
+	// Roughly: VersionInfo + PackageDataOffset + Footer.
+	return sizeof(uint32) + sizeof(int64) + sizeof(FGuid);
+}
 
-	// Test the footer is in place so we know we didn't crash mid-write
-	bool bParsedFooter = false;
+/**
+ * Store the package data, possibly compressed, along with some meta data.
+ * @param PackageDataStream The source stream containing the package data.
+ * @param PackageDataFormat How the package data should be stored in the blob.
+ * @param DstAr The archive used to write the blob data.
+ */
+bool WritePackageBlob(FConcertPackageDataStream& PackageDataStream, EPackageDataFormat PackageDataFormat, FArchive& DstAr)
+{
+	check(DstAr.IsSaving() && DstAr.Tell() == 0);
+	check(PackageDataStream.DataAr == nullptr || PackageDataStream.DataAr->IsLoading()); // The package data stream may be null to write an 'empty' package.
+	check(PackageDataStream.DataAr != nullptr || (PackageDataStream.DataAr == nullptr && PackageDataStream.DataSize == 0 && PackageDataStream.DataBlob == nullptr)); // If the package data archive is null, the size must be zero.
+
+	// Write blob version and format.
+	WritePackageHeader(DstAr, PackageDataFormat);
+
+	if (PackageDataFormat == EPackageDataFormat::PackageDataCompressed)
 	{
-		const int64 SerializedTransactionSize = Ar.TotalSize();
-		if (SerializedTransactionSize >= sizeof(FGuid))
+		// Write the package data size.
+		uint32 UncompressedPackageSize = PackageDataStream.DataSize;
+		DstAr.SerializeIntPacked(UncompressedPackageSize);
+		check(static_cast<uint64>(UncompressedPackageSize) <= TNumericLimits<TArray<uint8>::SizeType>::Max());
+
+		// Compress and write the package data.
+		if (UncompressedPackageSize > 0)
 		{
-			FGuid SerializedFooter;
-			Ar.Seek(SerializedTransactionSize - sizeof(FGuid));
-			Ar << SerializedFooter;
-			Ar.Seek(0);
-			bParsedFooter = SerializedFooter == EntryFooter;
+			if (PackageDataStream.DataBlob) // Optimization to avoid transfering from the archive to a temporary buffer.
+			{
+				DstAr.SerializeCompressed(const_cast<uint8*>(PackageDataStream.DataBlob->GetData()), UncompressedPackageSize, NAME_Zlib); // Write
+			}
+			else
+			{
+				TArray<uint8> UncompressedPackageData;
+				UncompressedPackageData.AddUninitialized(UncompressedPackageSize);
+				PackageDataStream.DataAr->Serialize(UncompressedPackageData.GetData(), UncompressedPackageSize); // Read.
+				DstAr.SerializeCompressed(const_cast<uint8*>(UncompressedPackageData.GetData()), UncompressedPackageSize, NAME_Zlib); // Write.
+			}
 		}
 	}
-	if (!bParsedFooter)
+	else
 	{
-		return false;
+		check(PackageDataFormat == EPackageDataFormat::PackageDataUncompressed);
+
+		// Write the package size.
+		DstAr << PackageDataStream.DataSize;
+
+		// Copy the package data in the destination.
+		if (PackageDataStream.DataSize > 0)
+		{
+			ConcertUtil::Copy(DstAr, *PackageDataStream.DataAr, PackageDataStream.DataSize);
+		}
 	}
 
-	// Deserialize the data version
-	uint32 SerializedDataVersion = 0;
-	Ar.SerializeIntPacked(SerializedDataVersion);
+	// Write a footer as proof that the operation did not crash mid-write.
+	WritePackageFooter(DstAr);
 
-	// Deserialize the info (header)
-	int64 BodyOffset = 0;
-	Ar << BodyOffset;
-	if (OutPackageInfo)
+	return !DstAr.IsError();
+}
+
+bool WritePackageBlob(FConcertPackageDataStream& PackageDataStream, EPackageDataFormat InPackageDataFormat, TArray<uint8>& OutPackageBlob)
+{
+	FMemoryWriter DstAr(OutPackageBlob);
+	return WritePackageBlob(PackageDataStream, InPackageDataFormat, DstAr);
+}
+
+/** Extract the package data in OutPackageData if the data is smaller than the streaming threshold, otherwise, stream the data using the streaming function.*/
+bool ExtractPackageData(FArchive& PackageBlobAr, const TFunctionRef<void(FConcertPackageDataStream&)>& PackageDataStreamFallbackFn)
+{
+	// Ensure the footer is in place so we know we didn't crash mid-write
+	if (!HasValidPackageFooter(PackageBlobAr))
 	{
-		FConcertPackageInfo::StaticStruct()->SerializeItem(Ar, OutPackageInfo, nullptr);
+		return false;
 	}
 
-	// Deserialize the raw data (body)
-	if (OutPackageData)
-	{
-		Ar.Seek(BodyOffset);
+	// Read the header.
+	TTuple<EPackageDataFormat, int64> Header = ReadPackageHeader(PackageBlobAr);
+	EPackageDataFormat PackageFormat = Header.Get<0>();
+	int64 PackageDataOffset = Header.Get<1>();
+
+	// Put the stream in position.
+	PackageBlobAr.Seek(PackageDataOffset);
 
+	// The package data is compressed (this hints the package is small enough to be loaded in memory)
+	if (PackageFormat == EPackageDataFormat::PackageDataCompressed)
+	{
+		// Read the uncompressed package data size.
 		uint32 UncompressedPackageSize = 0;
-		Ar.SerializeIntPacked(UncompressedPackageSize);
-		OutPackageData->Reset(UncompressedPackageSize);
-		OutPackageData->AddZeroed(UncompressedPackageSize);
+		PackageBlobAr.SerializeIntPacked(UncompressedPackageSize);
+		check(static_cast<uint64>(UncompressedPackageSize) <= TNumericLimits<TArray<uint8>::SizeType>::Max()); // Prevent overflowing the variable holding the number of element in the array.
+
+		// Decompress the data (in memory).
+		TArray<uint8> DecompressedData;
 		if (UncompressedPackageSize > 0)
 		{
-			Ar.SerializeCompressed(OutPackageData->GetData(), UncompressedPackageSize, NAME_Zlib);
+			DecompressedData.AddUninitialized(UncompressedPackageSize);
+			PackageBlobAr.SerializeCompressed(DecompressedData.GetData(), UncompressedPackageSize, NAME_Zlib); // Read and decompress the data.
 		}
+
+		// Let the caller stream the package data from a memory archive.
+		FMemoryReader DataAr(DecompressedData);
+		FConcertPackageDataStream PackageDataStream{ &DataAr, UncompressedPackageSize, &DecompressedData };
+		PackageDataStreamFallbackFn(PackageDataStream);
+	}
+	else
+	{
+		check(PackageFormat == EPackageDataFormat::PackageDataUncompressed);
+
+		// Read the package size.
+		int64 PackageDataSize = 0;
+		PackageBlobAr << PackageDataSize;
+
+		// Let the caller stream the package data from the archive.
+		FConcertPackageDataStream PackageDataStream{ &PackageBlobAr, PackageDataSize, nullptr };
+		PackageDataStreamFallbackFn(PackageDataStream);
 	}
 
-	return !Ar.IsError();
+	return !PackageBlobAr.IsError();
+}
+
+bool ExtractPackageData(const TArray<uint8>& InPackageBlob, const TFunctionRef<void(FConcertPackageDataStream&)>& PackageDataStreamFallbackFn)
+{
+	FMemoryReader PackageBlobAr(InPackageBlob);
+	return ExtractPackageData(PackageBlobAr, PackageDataStreamFallbackFn);
+}
+
+bool ExtractPackageData(const FString& InPackageBlobPathname, const TFunctionRef<void(FConcertPackageDataStream&)>& PackageDataStreamFallbackFn)
+{
+	TUniquePtr<FArchive> PackageBlobAr(IFileManager::Get().CreateFileReader(*InPackageBlobPathname)); // The file containing the database package data (along with the package info and some meta data)
+	return PackageBlobAr ? ExtractPackageData(*PackageBlobAr, PackageDataStreamFallbackFn) : false;
 }
 
 } // namespace PackageDataUtil
@@ -1199,8 +1343,11 @@ private:
 /** Defined here where TUniquePtr can see the definition of FConcertFileCache and FConcertSyncSessionDatabaseStatements as the TUniquePtr constructor/destructor cannot work with a forward declared type */
 FConcertSyncSessionDatabase::FConcertSyncSessionDatabase()
 	: Database(MakeUnique<FSQLiteDatabase>())
+	, MaxPackageBlobSizeForCaching(FMath::Max(PackageDataUtil::MaxFileSizeBytesToCache / PackageDataUtil::MinFilesToCache, 16ull * 1024 * 1024)) // Prevent caching very large files.
+	, MaxPackageDataSizeForCompression(32ull * 1024 * 1024) // Bigger files will not be compressed to avoid performance hit to decompress them when requested. (Streaming compression is not supported yet, so it is done in memory only)
 {
 }
+
 FConcertSyncSessionDatabase::~FConcertSyncSessionDatabase() = default;
 
 bool FConcertSyncSessionDatabase::IsValid() const
@@ -1384,13 +1531,13 @@ bool FConcertSyncSessionDatabase::AddTransactionActivity(const FConcertSyncTrans
 		);
 }
 
-bool FConcertSyncSessionDatabase::AddPackageActivity(const FConcertSyncPackageActivity& InPackageActivity, int64& OutActivityId, int64& OutPackageEventId)
+bool FConcertSyncSessionDatabase::AddPackageActivity(const FConcertSyncActivity& InPackageActivityBasePart, const FConcertPackageInfo& PackageInfo, FConcertPackageDataStream& InPackageDataStream, int64& OutActivityId, int64& OutPackageEventId)
 {
 	FConcertSyncSessionDatabaseScopedTransaction ScopedTransaction(*Statements);
 	return ScopedTransaction.CommitOrRollback(
-		AddPackageEvent(InPackageActivity.EventData, OutPackageEventId) &&
-		Statements->AddActivityData(InPackageActivity.EndpointId, EConcertSyncActivityEventType::Package, OutPackageEventId, InPackageActivity.EventSummary, OutActivityId) &&
-		SetActivityIgnoredState(OutActivityId, InPackageActivity.bIgnored)
+		AddPackageEvent(PackageInfo, InPackageDataStream, OutPackageEventId) &&
+		Statements->AddActivityData(InPackageActivityBasePart.EndpointId, EConcertSyncActivityEventType::Package, OutPackageEventId, InPackageActivityBasePart.EventSummary, OutActivityId) &&
+		SetActivityIgnoredState(OutActivityId, InPackageActivityBasePart.bIgnored)
 		);
 }
 
@@ -1424,13 +1571,13 @@ bool FConcertSyncSessionDatabase::SetTransactionActivity(const FConcertSyncTrans
 		);
 }
 
-bool FConcertSyncSessionDatabase::SetPackageActivity(const FConcertSyncPackageActivity& InPackageActivity, const bool bMetaDataOnly)
+bool FConcertSyncSessionDatabase::SetPackageActivity(const FConcertSyncActivity& InPackageActivityBasePart, FConcertSyncPackageEventData& InPackageActivityEventPart, const bool bMetaDataOnly)
 {
 	FConcertSyncSessionDatabaseScopedTransaction ScopedTransaction(*Statements);
 	return ScopedTransaction.CommitOrRollback(
-		SetPackageEvent(InPackageActivity.EventId, InPackageActivity.EventData, bMetaDataOnly) &&
-		Statements->SetActivityData(InPackageActivity.ActivityId, InPackageActivity.EndpointId, InPackageActivity.EventTime, InPackageActivity.EventType, InPackageActivity.EventId, InPackageActivity.EventSummary) &&
-		SetActivityIgnoredState(InPackageActivity.ActivityId, InPackageActivity.bIgnored)
+		SetPackageEvent(InPackageActivityBasePart.EventId, InPackageActivityEventPart.MetaData.PackageRevision, InPackageActivityEventPart.MetaData.PackageInfo, bMetaDataOnly ? static_cast<FConcertPackageDataStream*>(nullptr) : &InPackageActivityEventPart.PackageDataStream) &&
+		Statements->SetActivityData(InPackageActivityBasePart.ActivityId, InPackageActivityBasePart.EndpointId, InPackageActivityBasePart.EventTime, InPackageActivityBasePart.EventType, InPackageActivityBasePart.EventId, InPackageActivityBasePart.EventSummary) &&
+		SetActivityIgnoredState(InPackageActivityBasePart.ActivityId, InPackageActivityBasePart.bIgnored)
 		);
 }
 
@@ -1459,10 +1606,17 @@ bool FConcertSyncSessionDatabase::GetTransactionActivity(const int64 InActivityI
 		&& GetTransactionEvent(OutTransactionActivity.EventId, OutTransactionActivity.EventData);
 }
 
-bool FConcertSyncSessionDatabase::GetPackageActivity(const int64 InActivityId, FConcertSyncPackageActivity& OutPackageActivity) const
+bool FConcertSyncSessionDatabase::GetPackageActivity(const int64 InActivityId, const TFunctionRef<void(FConcertSyncActivity&&, FConcertSyncPackageEventData&)>& PackageActivityFn) const
 {
-	return GetActivity(InActivityId, OutPackageActivity)
-		&& GetPackageEvent(OutPackageActivity.EventId, OutPackageActivity.EventData);
+	FConcertSyncActivity PackageActivityBasePart;
+
+	auto GetPackageActivityEventFn = [&PackageActivityFn, &PackageActivityBasePart](FConcertSyncPackageEventData& PackageActivityEventPart) mutable
+	{
+		PackageActivityFn(MoveTemp(PackageActivityBasePart), PackageActivityEventPart);
+	};
+
+	return GetActivity(InActivityId, PackageActivityBasePart) // Pull the base part of the package activity.
+		&& GetPackageEvent(PackageActivityBasePart.EventId, GetPackageActivityEventFn); // Pull the event part of the activity and call back (to let the caller stream the package data).
 }
 
 bool FConcertSyncSessionDatabase::GetActivityEventType(const int64 InActivityId, EConcertSyncActivityEventType& OutEventType) const
@@ -1500,10 +1654,17 @@ bool FConcertSyncSessionDatabase::GetTransactionActivityForEvent(const int64 InT
 		&& GetTransactionEvent(InTransactionEventId, OutTransactionActivity.EventData);
 }
 
-bool FConcertSyncSessionDatabase::GetPackageActivityForEvent(const int64 InPackageEventId, FConcertSyncPackageActivity& OutPackageActivity) const
+bool FConcertSyncSessionDatabase::GetPackageActivityForEvent(const int64 InPackageEventId, const TFunctionRef<void(FConcertSyncActivity&&, FConcertSyncPackageEventData&)>& PackageActivityFn) const
 {
-	return GetActivityForEvent(InPackageEventId, EConcertSyncActivityEventType::Package, OutPackageActivity)
-		&& GetPackageEvent(InPackageEventId, OutPackageActivity.EventData);
+	FConcertSyncActivity PackageActivityBasePart;
+
+	auto GetPackageActivityEventFn = [&PackageActivityFn, &PackageActivityBasePart](FConcertSyncPackageEventData& PackageActivityEventPart) mutable
+	{
+		PackageActivityFn(MoveTemp(PackageActivityBasePart), PackageActivityEventPart);
+	};
+
+	return GetActivityForEvent(InPackageEventId, EConcertSyncActivityEventType::Package, PackageActivityBasePart) // Pull the base activity part (common to all activity)
+		&& GetPackageEvent(InPackageEventId, GetPackageActivityEventFn); // Pull the package specific part and callback
 }
 
 bool FConcertSyncSessionDatabase::EnumerateActivities(TFunctionRef<bool(FConcertSyncActivity&&)> InCallback) const
@@ -1590,11 +1751,11 @@ bool FConcertSyncSessionDatabase::EnumerateTransactionActivities(TFunctionRef<bo
 	});
 }
 
-bool FConcertSyncSessionDatabase::EnumeratePackageActivities(TFunctionRef<bool(FConcertSyncPackageActivity&&)> InCallback) const
+bool FConcertSyncSessionDatabase::EnumeratePackageActivities(const TFunctionRef<bool(FConcertSyncActivity&&, FConcertSyncPackageEventData&)>& InCallback) const
 {
 	return Statements->GetAllActivityDataForEventType(EConcertSyncActivityEventType::Package, [this, &InCallback](const int64 InActivityId, const FGuid& InEndpointId, const FDateTime InEventTime, const int64 InEventId, FConcertSessionSerializedCborPayload&& InEventSummary)
 	{
-		FConcertSyncPackageActivity PackageActivity;
+		FConcertSyncActivity PackageActivity;
 		PackageActivity.ActivityId = InActivityId;
 		PackageActivity.bIgnored = Statements->IsActivityIgnored(InActivityId);
 		PackageActivity.EndpointId = InEndpointId;
@@ -1602,13 +1763,15 @@ bool FConcertSyncSessionDatabase::EnumeratePackageActivities(TFunctionRef<bool(F
 		PackageActivity.EventType = EConcertSyncActivityEventType::Package;
 		PackageActivity.EventId = InEventId;
 		PackageActivity.EventSummary = MoveTemp(InEventSummary);
-		if (GetPackageEvent(PackageActivity.EventId, PackageActivity.EventData))
+
+		ESQLitePreparedStatementExecuteRowResult Result = ESQLitePreparedStatementExecuteRowResult::Error;
+		GetPackageEvent(PackageActivity.EventId, [&InCallback, &Result, BasePart = MoveTemp(PackageActivity)](FConcertSyncPackageEventData& PackageEventPart) mutable
 		{
-			return InCallback(MoveTemp(PackageActivity))
+			Result = InCallback(MoveTemp(BasePart), PackageEventPart)
 				? ESQLitePreparedStatementExecuteRowResult::Continue
 				: ESQLitePreparedStatementExecuteRowResult::Stop;
-		}
-		return ESQLitePreparedStatementExecuteRowResult::Error;
+		});
+		return Result;
 	});
 }
 
@@ -1978,18 +2141,19 @@ bool FConcertSyncSessionDatabase::AddDummyPackageEvent(const FName InPackageName
 			{
 				// If the head package event is a dummy event with no activity associated with it, we'll re-use the head package event, otherwise we'll add a new one
 				FString FoundDataFilename;
-				FConcertSyncPackageEvent FoundPackageEvent;
-				if (Statements->GetPackageEventForId(OutPackageEventId, FoundPackageEvent.PackageRevision, FoundPackageEvent.Package.Info, FoundDataFilename))
+				FConcertPackageInfo PackageInfo;
+				int64 PackageRevision;
+				if (Statements->GetPackageEventForId(OutPackageEventId, PackageRevision, PackageInfo, FoundDataFilename))
 				{
-					if (FoundPackageEvent.Package.Info.PackageUpdateType == EConcertPackageUpdateType::Dummy)
+					if (PackageInfo.PackageUpdateType == EConcertPackageUpdateType::Dummy)
 					{
 						// Does this package have associated activity? If so, we need to keep it as-is
 						FConcertSyncActivity FoundActivity;
 						if (!GetActivityForEvent(OutPackageEventId, EConcertSyncActivityEventType::Package, FoundActivity))
 						{
 							// Update this package event
-							return Statements->GetTransactionMaxEventId(FoundPackageEvent.Package.Info.TransactionEventIdAtSave)
-								&& Statements->SetPackageEvent(OutPackageEventId, PackageNameId, FoundPackageEvent.PackageRevision, FoundPackageEvent.Package.Info.TransactionEventIdAtSave, FoundPackageEvent.Package.Info, FoundDataFilename);
+							return Statements->GetTransactionMaxEventId(PackageInfo.TransactionEventIdAtSave)
+								&& Statements->SetPackageEvent(OutPackageEventId, PackageNameId, PackageRevision, PackageInfo.TransactionEventIdAtSave, PackageInfo, FoundDataFilename);
 						}
 					}
 				}
@@ -1997,18 +2161,19 @@ bool FConcertSyncSessionDatabase::AddDummyPackageEvent(const FName InPackageName
 		}
 
 		// Add a new package event
-		FConcertSyncPackageEvent DummyPackageEvent;
-		DummyPackageEvent.Package.Info.PackageName = InPackageName;
-		DummyPackageEvent.Package.Info.PackageUpdateType = EConcertPackageUpdateType::Dummy;
-		return Statements->GetTransactionMaxEventId(DummyPackageEvent.Package.Info.TransactionEventIdAtSave)
-			&& AddPackageEvent(DummyPackageEvent, OutPackageEventId);
+		FConcertPackageInfo DummyPackageInfo;
+		FConcertPackageDataStream DummyPackageDataStream;
+		DummyPackageInfo.PackageName = InPackageName;
+		DummyPackageInfo.PackageUpdateType = EConcertPackageUpdateType::Dummy;
+		return Statements->GetTransactionMaxEventId(DummyPackageInfo.TransactionEventIdAtSave)
+			&& AddPackageEvent(DummyPackageInfo, DummyPackageDataStream, OutPackageEventId);
 	};
 
 	FConcertSyncSessionDatabaseScopedTransaction ScopedTransaction(*Statements);
 	return ScopedTransaction.CommitOrRollback(AddDummyPackageEventImpl());
 }
 
-bool FConcertSyncSessionDatabase::AddPackageEvent(const FConcertSyncPackageEvent& InPackageEvent, int64& OutPackageEventId)
+bool FConcertSyncSessionDatabase::AddPackageEvent(const FConcertPackageInfo& InPackageInfo, FConcertPackageDataStream& InPackageDataStream, int64& OutPackageEventId)
 {
 	// Get the next package ID
 	if (!GetPackageMaxEventId(OutPackageEventId) || OutPackageEventId == MAX_int64)
@@ -2019,84 +2184,81 @@ bool FConcertSyncSessionDatabase::AddPackageEvent(const FConcertSyncPackageEvent
 
 	// Get the next package revision
 	int64 PackageRevision = 0;
-	if (!GetPackageHeadRevision(InPackageEvent.Package.Info.PackageName, PackageRevision) || PackageRevision == MAX_int64)
+	if (!GetPackageHeadRevision(InPackageInfo.PackageName, PackageRevision) || PackageRevision == MAX_int64)
 	{
 		return false;
 	}
 	++PackageRevision;
 
-	return SetPackageEvent(OutPackageEventId, PackageRevision, InPackageEvent.Package);
+	return SetPackageEvent(OutPackageEventId, PackageRevision, InPackageInfo, &InPackageDataStream);
 }
 
-bool FConcertSyncSessionDatabase::UpdatePackageEvent(const int64 InPackageEventId, const FConcertSyncPackageEvent& InPackageEvent)
+bool FConcertSyncSessionDatabase::UpdatePackageEvent(const int64 InPackageEventId, FConcertSyncPackageEventData& PackageEventSrc)
 {
 	int64 MaxPackageEventId;
 	if (GetPackageMaxEventId(MaxPackageEventId) && InPackageEventId <= MaxPackageEventId) // Ensure the package ID is in bound.
 	{
 		FConcertSyncSessionDatabaseScopedTransaction ScopedTransaction(*Statements);
 		return ScopedTransaction.CommitOrRollback(
-			SetPackageEvent(InPackageEventId, InPackageEvent)
+			SetPackageEvent(InPackageEventId, PackageEventSrc.MetaData.PackageRevision, PackageEventSrc.MetaData.PackageInfo, &PackageEventSrc.PackageDataStream)
 		);
 	}
 
 	return false;
 }
 
-bool FConcertSyncSessionDatabase::SetPackageEvent(const int64 InPackageEventId, const FConcertSyncPackageEvent& InPackageEvent, const bool bMetaDataOnly)
-{
-	return SetPackageEvent(InPackageEventId, InPackageEvent.PackageRevision, InPackageEvent.Package, bMetaDataOnly);
-}
-
-bool FConcertSyncSessionDatabase::SetPackageEvent(const int64 InPackageEventId, const int64 InPackageRevision, const FConcertPackage& InPackage, const bool bMetaDataOnly)
+bool FConcertSyncSessionDatabase::SetPackageEvent(const int64 InPackageEventId, int64 InPackageRevision, const FConcertPackageInfo& InPackageInfo, FConcertPackageDataStream* InPackageDataStream)
 {
 	if (!ensureAlwaysMsgf(InPackageRevision > 0, TEXT("Invalid package revision! Must be greater than zero.")))
 	{
 		return false;
 	}
-	if (!ensureAlwaysMsgf(!InPackage.Info.PackageName.IsNone(), TEXT("Invalid package name! Must be set.")))
+	if (!ensureAlwaysMsgf(!InPackageInfo.PackageName.IsNone(), TEXT("Invalid package name! Must be set.")))
 	{
 		return false;
 	}
 
 	// Ensure an entry for this package name
 	int64 PackageNameId = 0;
-	if (!EnsurePackageNameId(InPackage.Info.PackageName, PackageNameId))
+	if (!EnsurePackageNameId(InPackageInfo.PackageName, PackageNameId))
 	{
 		return false;
 	}
 
 	// Write the data blob file
-	const FString PackageDataFilename = PackageDataUtil::GetDataFilename(InPackage.Info.PackageName, InPackageRevision);
+	const FString PackageDataFilename = PackageDataUtil::GetDataFilename(InPackageInfo.PackageName, InPackageRevision);
 
-	if (!bMetaDataOnly)
+	if (InPackageDataStream) // If !MetaDataOnly
 	{
-		const FString PackageDataPathname = PackageDataUtil::GetDataPath(SessionPath) / PackageDataFilename;
-		if (!SavePackage(PackageDataPathname, InPackage.Info, InPackage.PackageData))
+		const FString DstPackagBlobPathname = PackageDataUtil::GetDataPath(SessionPath) / PackageDataFilename;
+		if (!SavePackage(DstPackagBlobPathname, InPackageInfo, *InPackageDataStream))
 		{
 			return false;
 		}
 	}
 
 	// Add the database entry
-	return Statements->SetPackageEvent(InPackageEventId, PackageNameId, InPackageRevision, InPackage.Info.TransactionEventIdAtSave, InPackage.Info, PackageDataFilename);
+	return Statements->SetPackageEvent(InPackageEventId, PackageNameId, InPackageRevision, InPackageInfo.TransactionEventIdAtSave, InPackageInfo, PackageDataFilename);
 }
 
-bool FConcertSyncSessionDatabase::GetPackageEvent(const int64 InPackageEventId, FConcertSyncPackageEvent& OutPackageEvent, const bool InMetaDataOnly) const
+bool FConcertSyncSessionDatabase::GetPackageEventMetaData(const int64 InPackageEventId, int64& OuptPackageRevision, FConcertPackageInfo& OutPackageInfo) const
+{
+	FString UnusedDataFilename;
+	return Statements->GetPackageEventForId(InPackageEventId, OuptPackageRevision, OutPackageInfo, UnusedDataFilename);
+}
+
+bool FConcertSyncSessionDatabase::GetPackageEvent(const int64 InPackageEventId, const TFunctionRef<void(FConcertSyncPackageEventData&)>& PackageEventFn) const
 {
 	FString DataFilename;
-	if (Statements->GetPackageEventForId(InPackageEventId, OutPackageEvent.PackageRevision, OutPackageEvent.Package.Info, DataFilename))
+	FConcertSyncPackageEventData PackageEventSrc;
+	if (Statements->GetPackageEventForId(InPackageEventId, PackageEventSrc.MetaData.PackageRevision, PackageEventSrc.MetaData.PackageInfo, DataFilename))
 	{
-		if (InMetaDataOnly)
-		{
-			OutPackageEvent.Package.PackageData.Reset();
-			return true;
-		}
-
 		const FString PackageDataPathname = PackageDataUtil::GetDataPath(SessionPath) / DataFilename;
-		if (LoadPackage(PackageDataPathname, nullptr, &OutPackageEvent.Package.PackageData))
+		return LoadPackage(PackageDataPathname, [&PackageEventSrc, &PackageEventFn](FConcertPackageDataStream& InPackageDataStream)
 		{
-			return true;
-		}
+			PackageEventSrc.PackageDataStream = InPackageDataStream;
+			PackageEventFn(PackageEventSrc);
+		});
 	}
 
 	return false;
@@ -2152,23 +2314,22 @@ bool FConcertSyncSessionDatabase::EnumeratePackageNamesWithHeadRevision(TFunctio
 	});
 }
 
-bool FConcertSyncSessionDatabase::EnumerateHeadRevisionPackageData(TFunctionRef<bool(FConcertPackage&&)> InCallback, const bool InMetaDataOnly) const
+bool FConcertSyncSessionDatabase::EnumerateHeadRevisionPackageData(TFunctionRef<bool(const FConcertPackageInfo&, FConcertPackageDataStream&)> InCallback) const
 {
-	return Statements->GetUniquePackageNameIdsForPackageEvents([this, &InCallback, InMetaDataOnly](int64 InPackageNameId)
+	return Statements->GetUniquePackageNameIdsForPackageEvents([this, &InCallback](int64 InPackageNameId)
 	{
 		int64 PackageHeadRevision = 0;
 		if (Statements->GetPackageHeadRevision(InPackageNameId, PackageHeadRevision))
 		{
 			FString DataFilename;
-			FConcertPackage Package;
-			if (Statements->GetPackageDataForRevision(InPackageNameId, PackageHeadRevision, Package.Info, DataFilename))
+			FConcertPackageInfo PackageInfo;
+			if (Statements->GetPackageDataForRevision(InPackageNameId, PackageHeadRevision, PackageInfo, DataFilename))
 			{
-				const FString PackageDataPathname = PackageDataUtil::GetDataPath(SessionPath) / DataFilename;
-				if (InMetaDataOnly || LoadPackage(PackageDataPathname, nullptr, &Package.PackageData))
+				bool bContinue = true;
+				const FString PackageBlobPathname = PackageDataUtil::GetDataPath(SessionPath) / DataFilename;
+				if (LoadPackage(PackageBlobPathname, [&bContinue, &InCallback, &PackageInfo](FConcertPackageDataStream& InPackageDataStream){ bContinue = InCallback(PackageInfo, InPackageDataStream); }))
 				{
-					return InCallback(MoveTemp(Package))
-						? ESQLitePreparedStatementExecuteRowResult::Continue
-						: ESQLitePreparedStatementExecuteRowResult::Stop;
+					return bContinue ? ESQLitePreparedStatementExecuteRowResult::Continue : ESQLitePreparedStatementExecuteRowResult::Stop;
 				}
 			}
 		}
@@ -2193,12 +2354,33 @@ bool FConcertSyncSessionDatabase::AddPersistEventForHeadRevision(FName InPackage
 	return false;
 }
 
-bool FConcertSyncSessionDatabase::GetPackageDataForRevision(const FName InPackageName, FConcertPackage& OutPackage, const int64* InPackageRevision) const
+bool FConcertSyncSessionDatabase::GetPackageInfoForRevision(const FName InPackageName, FConcertPackageInfo& OutPackageInfo, const int64* InPackageRevision) const
 {
-	return GetPackageDataForRevision(InPackageName, OutPackage.Info, &OutPackage.PackageData, InPackageRevision);
+	int64 PackageRevision = 0;
+	if (InPackageRevision)
+	{
+		PackageRevision = *InPackageRevision;
+	}
+	else if (!GetPackageHeadRevision(InPackageName, PackageRevision))
+	{
+		return false;
+	}
+	if (PackageRevision == 0)
+	{
+		return false;
+	}
+
+	int64 PackageNameId = 0;
+	if (!GetPackageNameId(InPackageName, PackageNameId))
+	{
+		return false;
+	}
+
+	FString DataFilename;
+	return Statements->GetPackageDataForRevision(PackageNameId, PackageRevision, OutPackageInfo, DataFilename);
 }
 
-bool FConcertSyncSessionDatabase::GetPackageDataForRevision(const FName InPackageName, FConcertPackageInfo& OutPackageInfo, TArray<uint8>* OutPackageData, const int64* InPackageRevision) const
+bool FConcertSyncSessionDatabase::GetPackageDataForRevision(const FName InPackageName, const TFunctionRef<void(const FConcertPackageInfo&, FConcertPackageDataStream&)>& InPackageDataCallback, const int64* InPackageRevision) const
 {
 	int64 PackageRevision = 0;
 	if (InPackageRevision)
@@ -2220,19 +2402,15 @@ bool FConcertSyncSessionDatabase::GetPackageDataForRevision(const FName InPackag
 		return false;
 	}
 
+	FConcertPackageInfo PackageInfo;
 	FString DataFilename;
-	if (Statements->GetPackageDataForRevision(PackageNameId, PackageRevision, OutPackageInfo, DataFilename))
+	if (Statements->GetPackageDataForRevision(PackageNameId, PackageRevision, PackageInfo, DataFilename))
 	{
-		if (!OutPackageData)
-		{
-			return true;
-		}
-
 		const FString PackageDataPathname = PackageDataUtil::GetDataPath(SessionPath) / DataFilename;
-		if (LoadPackage(PackageDataPathname, nullptr, OutPackageData))
+		return LoadPackage(PackageDataPathname, [&PackageInfo, &InPackageDataCallback](FConcertPackageDataStream& InPackageDataStream)
 		{
-			return true;
-		}
+			InPackageDataCallback(PackageInfo, InPackageDataStream);
+		});
 	}
 
 	return false;
@@ -2351,16 +2529,74 @@ bool FConcertSyncSessionDatabase::LoadTransaction(const FString& InTransactionFi
 	return false;
 }
 
-bool FConcertSyncSessionDatabase::SavePackage(const FString& InPackageFilename, const FConcertPackageInfo& InPackageInfo, const TArray<uint8>& InPackageData) const
+bool FConcertSyncSessionDatabase::SavePackage(const FString& InDstPackageBlobPathname, const FConcertPackageInfo& InPackageInfo, FConcertPackageDataStream& InPackageDataStream) const
+{
+	// Decide if the package data should be stored compressed.
+	PackageDataUtil::EPackageDataFormat PackageDataFormat = ShouldCompressPackageData(InPackageDataStream.DataSize) ? PackageDataUtil::EPackageDataFormat::PackageDataCompressed : PackageDataUtil::EPackageDataFormat::PackageDataUncompressed;
+	if (PackageDataFormat == PackageDataUtil::EPackageDataFormat::PackageDataCompressed)
+	{
+		TArray<uint8> PackageBlob;
+		if (PackageDataUtil::WritePackageBlob(InPackageDataStream, PackageDataFormat, PackageBlob))
+		{
+			if (ShouldCachePackageBlob(PackageBlob.Num()))
+			{
+				return PackageFileCache->SaveAndCacheFile(InDstPackageBlobPathname, MoveTemp(PackageBlob));
+			}
+
+			TUniquePtr<FArchive> DstAr(IFileManager::Get().CreateFileWriter(*InDstPackageBlobPathname));
+			if (DstAr)
+			{
+				DstAr->Serialize(PackageBlob.GetData(), PackageBlob.Num());
+				return DstAr->IsError();
+			}
+		}
+	}
+	else
+	{
+		check(PackageDataFormat == PackageDataUtil::EPackageDataFormat::PackageDataUncompressed);
+
+		int64 EstimatedPackageBlobSize = InPackageDataStream.DataSize + PackageDataUtil::EstimatePackageBlobMetaDataSize();
+		if (ShouldCachePackageBlob(EstimatedPackageBlobSize))
+		{
+			TArray<uint8> PackageBlob;
+			return PackageDataUtil::WritePackageBlob(InPackageDataStream, PackageDataFormat, PackageBlob) && PackageFileCache->SaveAndCacheFile(InDstPackageBlobPathname, MoveTemp(PackageBlob));
+		}
+		else
+		{
+			TUniquePtr<FArchive> DstAr(IFileManager::Get().CreateFileWriter(*InDstPackageBlobPathname));
+			return DstAr && PackageDataUtil::WritePackageBlob(InPackageDataStream, PackageDataFormat, *DstAr);
+		}
+	}
+
+	return false;
+}
+
+bool FConcertSyncSessionDatabase::LoadPackage(const FString& InPackageBlobFilename, const TFunctionRef<void(FConcertPackageDataStream&)>& PackageDataStreamFn) const
+{
+	int64 PackageBlobSize = IFileManager::Get().FileSize(*InPackageBlobFilename);
+	if (PackageBlobSize < 0) // Possible if the file doesn't exist.
+	{
+		UE_LOG(LogConcert, Warning, TEXT("Failed to read the size of database file '%s'."), *InPackageBlobFilename);
+		return false;
+	}
+
+	if (ShouldCachePackageBlob(static_cast<uint64>(PackageBlobSize))) // Try the cache if the blob is small enough to be found or stored.
+	{
+		TArray<uint8> PackageBlob;
+		return PackageFileCache->FindOrCacheFile(InPackageBlobFilename, PackageBlob) && PackageDataUtil::ExtractPackageData(PackageBlob, PackageDataStreamFn);
+	}
+
+	return PackageDataUtil::ExtractPackageData(InPackageBlobFilename, PackageDataStreamFn);
+}
+
+bool FConcertSyncSessionDatabase::ShouldCachePackageBlob(uint64 PackageBlobSize) const
 {
-	TArray<uint8> SerializedPackageData;
-	return PackageDataUtil::WritePackage(InPackageInfo, InPackageData, SerializedPackageData) && PackageFileCache->SaveAndCacheFile(InPackageFilename, MoveTemp(SerializedPackageData));
+	return PackageBlobSize <= MaxPackageBlobSizeForCaching;
 }
 
-bool FConcertSyncSessionDatabase::LoadPackage(const FString& InPackageFilename, FConcertPackageInfo* OutPackageInfo, TArray<uint8>* OutPackageData) const
+bool FConcertSyncSessionDatabase::ShouldCompressPackageData(uint64 PackageDataSize) const
 {
-	TArray<uint8> SerializedPackageData;
-	return PackageFileCache->FindOrCacheFile(InPackageFilename, SerializedPackageData) && PackageDataUtil::ReadPackage(SerializedPackageData, OutPackageInfo, OutPackageData);
+	return PackageDataSize <= MaxPackageDataSizeForCompression;
 }
 
 bool ConcertSyncSessionDatabaseFilterUtil::TransactionEventPassesFilter(const int64 InTransactionEventId, const FConcertSessionFilter& InSessionFilter, const FConcertSyncSessionDatabase& InDatabase)
diff --git a/Engine/Plugins/Developer/Concert/ConcertSync/ConcertSyncCore/Source/ConcertSyncCore/Public/ConcertLocalFileSharingService.h b/Engine/Plugins/Developer/Concert/ConcertSync/ConcertSyncCore/Source/ConcertSyncCore/Public/ConcertLocalFileSharingService.h
new file mode 100644
--- /dev/null
+++ b/Engine/Plugins/Developer/Concert/ConcertSync/ConcertSyncCore/Source/ConcertSyncCore/Public/ConcertLocalFileSharingService.h
@@ -0,0 +1,37 @@
+// Copyright Epic Games, Inc. All Rights Reserved.
+
+#pragma once
+
+#include "CoreMinimal.h"
+#include "IConcertFileSharingService.h"
+
+/**
+ * Share files using a local directory. It works only for one client and one server running on the same machine.
+ * @note This is mainly designed for the recovery system.
+ */
+class CONCERTSYNCCORE_API FConcertLocalFileSharingService : public IConcertFileSharingService
+{
+public:
+	/**
+	 * Construct the local file sharing service.
+	 * @param Role Appended to the shared directory (like DisasterRecovery) as a hint for the temp folder purpose. Client and server must use the same value.
+	 * @note The service uses the Engine intermediate directory to share the files. The files are automatically deleted after consumption. In case of crash
+	 *       the temporary shared files are deleted on next reboot if no other client/server on the machine are not running.
+	 */
+	FConcertLocalFileSharingService(const FString& Role);
+	virtual ~FConcertLocalFileSharingService();
+
+	virtual bool Publish(const FString& Pathname, FString& OutFileId) override;
+	virtual bool Publish(FArchive& SrcAr, int64 Size, FString& OutFileUri) override;
+	virtual TSharedPtr<FArchive> CreateReader(const FString& InFileUri) override;
+
+private:
+	void LoadActiveServices(TArray<uint32>& OutPids);
+	void SaveActiveServices(const TArray<uint32>& InPids);
+	void RemoveDeadProcessesAndFiles(TArray<uint32>& InOutPids);
+
+private:
+	FString SharedRootPathname;
+	FString SystemMutexName;
+	FString ActiveServicesRepositoryPathname;
+};
diff --git a/Engine/Plugins/Developer/Concert/ConcertSync/ConcertSyncCore/Source/ConcertSyncCore/Public/ConcertSyncSessionDatabase.h b/Engine/Plugins/Developer/Concert/ConcertSync/ConcertSyncCore/Source/ConcertSyncCore/Public/ConcertSyncSessionDatabase.h
--- a/Engine/Plugins/Developer/Concert/ConcertSync/ConcertSyncCore/Source/ConcertSyncCore/Public/ConcertSyncSessionDatabase.h
+++ b/Engine/Plugins/Developer/Concert/ConcertSync/ConcertSyncCore/Source/ConcertSyncCore/Public/ConcertSyncSessionDatabase.h
@@ -113,13 +113,15 @@ public:
 	 * @note The endpoint ID referenced by the activity must exist in the database (@see SetEndpoint).
 	 * @note This function is expected to be called on the server to populate its version of the session database.
 	 *
-	 * @param InPackageActivity			The package activity to add (the ActivityId, EventTime, EventType, EventId, and PackageRevision members are ignored).
+	 * @param InPackageActivity			The base part of the package activity to add (the ActivityId, EventTime, EventType, EventId, and PackageRevision members are ignored).
+	 * @param InPackageInfo				Information about the package affected.
+	 * @param InPackageDataStream		A stream to the package data if the activty is about adding or modifying a package. Can be empty when the activity is about deleting or discarding changes.
 	 * @param OutActivityId				Populated with the ID of this activity in the database.
 	 * @param OutPackageEventId			Populated with the ID of the package event in the database (@see GetPackageEvent).
 	 *
 	 * @return True if the package activity was added, false otherwise.
 	 */
-	bool AddPackageActivity(const FConcertSyncPackageActivity& InPackageActivity, int64& OutActivityId, int64& OutPackageEventId);
+	bool AddPackageActivity(const FConcertSyncActivity& InPackageActivity, const FConcertPackageInfo& InPackageInfo, FConcertPackageDataStream& InPackageDataStream, int64& OutActivityId, int64& OutPackageEventId);
 
 	/**
 	 * Set a connection activity in this database, creating or replacing it.
@@ -160,12 +162,13 @@ public:
 	 * @note The endpoint ID referenced by the activity must exist in the database (@see SetEndpoint).
 	 * @note This function is expected to be called on the client to populate its version of the session database from data synced from the server.
 	 *
-	 * @param InPackageActivity		The package activity to set.
-	 * @param bMetaDataOnly			True to store the meta data only, omitting the package data itself.
+	 * @param InPackageActivity			The base part of the package activity to set.
+	 * @param InPackageActivityEvent	The event meta data (package info + revision) and possibly the corresponding package data.
+	 * @param bMetaDataOnly				True to store the meta data only, omitting the package data itself.
 	 *
 	 * @return True if the package activity was set, false otherwise.
 	 */
-	bool SetPackageActivity(const FConcertSyncPackageActivity& InPackageActivity, const bool bMetaDataOnly = false);
+	bool SetPackageActivity(const FConcertSyncActivity& InPackageActivity, FConcertSyncPackageEventData& InPackageActivityEvent, const bool bMetaDataOnly = false);
 
 	/**
 	 * Get the generic part of an activity from this database.
@@ -210,12 +213,13 @@ public:
 	/**
 	 * Get a package activity from this database.
 	 *
-	 * @param InActivityId				The ID of the activity to find.
-	 * @param OutPackageActivity		The package activity to populate with the result.
+	 * @param InActivityId			The ID of the activity to find.
+	 * @param PackageActivityFn		A callback invoked with the package activity components (the base activity part and the event part).
+	 * @note The package data stream passed to the callback through the FConcertSycnPackageEventData is only valid during the callback.
 	 *
 	 * @return True if the package activity was found, false otherwise.
 	 */
-	bool GetPackageActivity(const int64 InActivityId, FConcertSyncPackageActivity& OutPackageActivity) const;
+	bool GetPackageActivity(const int64 InActivityId, const TFunctionRef<void(FConcertSyncActivity&& /*BasePart*/, FConcertSyncPackageEventData& /*EventPart*/)>& PackageActivityFn) const;
 
 	/**
 	 * Get the type of an activity in this database.
@@ -271,12 +275,13 @@ public:
 	/**
 	 * Get a package activity for an event in this database.
 	 *
-	 * @param InPackageEventId			The ID of the package event to find the activity for.
-	 * @param OutPackageActivity		The package activity to populate with the result.
+	 * @param InPackageEventId		The ID of the package event to find the activity for.
+	 * @param PackageActivityFn		A callback invoked with the package activity components (the base activity part and the event part).
+	 * @note The package data stream passed to the callback through the FConcertSycnPackageEventData is only valid during the callback.
 	 *
 	 * @return True if the package activity was found, false otherwise.
 	 */
-	bool GetPackageActivityForEvent(const int64 InPackageEventId, FConcertSyncPackageActivity& OutPackageActivity) const;
+	bool GetPackageActivityForEvent(const int64 InPackageEventId, const TFunctionRef<void(FConcertSyncActivity&& /*BasePart*/, FConcertSyncPackageEventData& /*EventPart*/)>& PackageActivityFn) const;
 
 	/**
 	 * Enumerate the generic part of the activities in this database.
@@ -321,7 +326,7 @@ public:
 	 *
 	 * @return True if the package activities were enumerated without error, false otherwise.
 	 */
-	bool EnumeratePackageActivities(TFunctionRef<bool(FConcertSyncPackageActivity&&)> InCallback) const;
+	bool EnumeratePackageActivities(const TFunctionRef<bool(FConcertSyncActivity&&/*BasePart*/, FConcertSyncPackageEventData& /*EventPart*/)>& InCallback) const;
 
 	/**
 	 * Enumerate all the activities in this database of the given type.
@@ -529,16 +534,26 @@ public:
 	 */
 	bool AddDummyPackageEvent(const FName InPackageName, int64& OutPackageEventId);
 
+	/**
+	 * Get a package event meta data (omitting the package data itself) from this database.
+	 *
+	 * @param InPackageEventId		The ID of the package event to find.
+	 * @param OuptPackageRevision	The package revision number.
+	 * @param OutPackageInfo		The Package info.
+	 *
+	 * @return True if the package event was found, false otherwise.
+	 */
+	bool GetPackageEventMetaData(const int64 InPackageEventId, int64& OuptPackageRevision, FConcertPackageInfo& OutPackageInfo) const;
+	
 	/**
 	 * Get a package event from this database.
 	 *
 	 * @param InPackageEventId			The ID of the package event to find.
-	 * @param OutPackageEvent			The package event to populate with the result.
-	 * @param InMetaDataOnly			True to only get the meta-data for the package, and skip the bulk of the data.
+	 * @param PackageEventFn			The callback invoked with the package event containing the event meta data and the package data if available.
 	 *
 	 * @return True if the package event was found, false otherwise.
 	 */
-	bool GetPackageEvent(const int64 InPackageEventId, FConcertSyncPackageEvent& OutPackageEvent, const bool InMetaDataOnly = false) const;
+	bool GetPackageEvent(const int64 InPackageEventId, const TFunctionRef<void(FConcertSyncPackageEventData&)>& PackageEventFn) const;
 
 	/**
 	 * Get package names for packages with a head revision (at least one package event)
@@ -564,34 +579,32 @@ public:
 	 * Enumerate the head revision package data for all packages in this database.
 	 *
 	 * @param InCallback				Callback invoked for each package; return true to continue enumeration, or false to stop.
-	 * @param InMetaDataOnly			True to only get the meta-data for the package, and skip the bulk of the data.
 	 *
 	 * @return True if the package data was enumerated without error, false otherwise.
 	 */
-	bool EnumerateHeadRevisionPackageData(TFunctionRef<bool(FConcertPackage&&)> InCallback, const bool InMetaDataOnly = false) const;
+	bool EnumerateHeadRevisionPackageData(TFunctionRef<bool(const FConcertPackageInfo&, FConcertPackageDataStream&)> InCallback) const;
 
 	/**
 	 * Get the data from this database for the given package name for the given revision.
 	 *
 	 * @param InPackageName				The name of the package to get the head revision for.
-	 * @param OutPackage				The package data to populate with the result.
+	 * @param OutPackageInfo			The package info to populate with the result.
 	 * @param InPackageRevision			The revision of the package to get the data for, or null to get the head revision.
 	 *
 	 * @return True if package data could be found for the given revision, false otherwise.
 	 */
-	bool GetPackageDataForRevision(const FName InPackageName, FConcertPackage& OutPackage, const int64* InPackageRevision = nullptr) const;
-
+	bool GetPackageInfoForRevision(const FName InPackageName, FConcertPackageInfo& OutPackageInfo, const int64* InPackageRevision = nullptr) const;
+	
 	/**
 	 * Get the data from this database for the given package name for the given revision.
 	 *
 	 * @param InPackageName				The name of the package to get the head revision for.
-	 * @param OutPackageInfo			The package info to populate with the result.
-	 * @param OutPackageData			The package data to populate with the result, or null to skip the package data.
+	 * @param InCallback				A callback invoked with the package info and package data.
 	 * @param InPackageRevision			The revision of the package to get the data for, or null to get the head revision.
 	 *
 	 * @return True if package data could be found for the given revision, false otherwise.
 	 */
-	bool GetPackageDataForRevision(const FName InPackageName, FConcertPackageInfo& OutPackageInfo, TArray<uint8>* OutPackageData, const int64* InPackageRevision = nullptr) const;
+	bool GetPackageDataForRevision(const FName InPackageName, const TFunctionRef<void(const FConcertPackageInfo&, FConcertPackageDataStream&)>& InCallback, const int64* InPackageRevision = nullptr) const;
 	
 	/**
 	 * Get the head revision in this database for the given package name.
@@ -636,10 +649,10 @@ public:
 	 * Update the specified package event.
 	 * @note The function is meant to be update a partially synced package event to store the corresponding package data.
 	 * @param InPackageEventId			The ID of the event to update.
-	 * @param InTransactionEvent		The event to store.
-	 * @return True if the transaction was updated.
+	 * @param PackageEvent				The event to store.
+	 * @return True if the package was updated.
 	 */
-	bool UpdatePackageEvent(const int64 InPackageEventId, const FConcertSyncPackageEvent& InPackageEvent);
+	bool UpdatePackageEvent(const int64 InPackageEventId, FConcertSyncPackageEventData& PackageEvent);
 
 private:
 	/**
@@ -726,35 +739,25 @@ private:
 	/**
 	 * Add a new package event to this database, assigning it a package event ID.
 	 *
-	 * @param InPackageEvent			The package event to add.
+	 * @param InPackageInfo				The package info to add for this event.
+	 * @param InPackageDataStream		The package data to add for this event (may be empty if the event doesn't carry package data).
 	 * @param OutPackageEventId			Populated with the ID of the package event in the database.
 	 *
 	 * @return True if the package event was added, false otherwise.
 	 */
-	bool AddPackageEvent(const FConcertSyncPackageEvent& InPackageEvent, int64& OutPackageEventId);
-
-	/**
-	 * Set a package event in this database, creating or replacing it.
-	 *
-	 * @param InPackageEventId			The ID of the package event to set.
-	 * @param InPackageEvent			The package event to set.
-	 * @param bMetaDataOnly				True to store the meta data only, omitting the package data.
-	 *
-	 * @return True if the package event was set, false otherwise.
-	 */
-	bool SetPackageEvent(const int64 InPackageEventId, const FConcertSyncPackageEvent& InPackageEvent, const bool bMetaDataOnly = false);
+	bool AddPackageEvent(const FConcertPackageInfo& InPackageInfo, FConcertPackageDataStream& InPackageDataStream, int64& OutPackageEventId);
 
 	/**
 	 * Set a package event in this database, creating or replacing it.
 	 *
 	 * @param InPackageEventId			The ID of the package event to set.
-	 * @param InPackageRevision			The package revision to set.
-	 * @param InPackage					The package to set.
-	 * @param bMetaDataOnly				True to store the meta data only, omitting the package data.
+	 * @param PackageRevision			The package revision number.
+	 * @param PackageInfo				The package info.
+	 * @param InPackageDataStream		The package data stream or null to set only the package meta data (revision and package info).
 	 *
 	 * @return True if the package event was set, false otherwise.
 	 */
-	bool SetPackageEvent(const int64 InPackageEventId, const int64 InPackageRevision, const FConcertPackage& InPackage, const bool bMetaDataOnly = false);
+	bool SetPackageEvent(const int64 InPackageEventId, int64 PackageRevision, const FConcertPackageInfo& PackageInfo, FConcertPackageDataStream* InPackageDataStream);
 
 	/**
 	 * Get the maximum ID of the package events in this database.
@@ -878,24 +881,37 @@ private:
 	/**
 	 * Save and cache the given package data with the given filename.
 	 *
-	 * @param InPackageFilename			The filename to save the package data as.
+	 * @param InDstPackageBlobPathname	The pathname to save the package data as. The package data may or may not be compressed.
 	 * @param InPackageInfo				The package info to save.
-	 * @param InPackageData				The package data to save.
+	 * @param InPackageDataStream		The package data to save.
 	 *
 	 * @return True if the package data was saved, false otherwise.
 	 */
-	bool SavePackage(const FString& InPackageFilename, const FConcertPackageInfo& InPackageInfo, const TArray<uint8>& InPackageData) const;
+	bool SavePackage(const FString& InDstPackageBlobPathname, const FConcertPackageInfo& InPackageInfo, FConcertPackageDataStream& InPackageDataStream) const;
 
 	/**
-	 * Load and cache the package data for the given filename.
-	 *
-	 * @param InPackageFilename			The filename to load the package data from.
-	 * @param OutPackageInfo			The package info to populate with the result, or null to skip the package info.
-	 * @param OutPackageData			The package data to populate with the result, or null to skip the package data.
+	 * Load the package data for the given filename.
 	 *
+	 * @param InPackageBlobFilename			The blob filename containing the package data to extract.
+	 * @param PackageDataStreamFn			Callback invoked to let the caller stream the package data.
+	 * @note The archive passed as output is only valid during the callback, its position is set at the beginning of the data (not necessarily zero) and may contain data passed the supplied size.
 	 * @return True if the package data was loaded, false otherwise.
 	 */
-	bool LoadPackage(const FString& InPackageFilename, FConcertPackageInfo* OutPackageInfo, TArray<uint8>* OutPackageData) const;
+	bool LoadPackage(const FString& InPackageBlobFilename, const TFunctionRef<void(FConcertPackageDataStream&)>& PackageDataStreamFn) const;
+	
+	/**
+	 * Returns true if a package blob should be cached in memory, according to its size.
+	 * @param PackageBlobSize The blob size (containing package data (possibly compressed) and some meta data).
+	 * @return true if the pacage should be cached, false otherwise.
+	 */
+	bool ShouldCachePackageBlob(uint64 PackageBlobSize) const;
+
+	/**
+	 * Returns true if the package data should be compressed in the blob, according to its size.
+	 * @param PackageDataSize The uncompressed package data size.
+	 * @return true if the package data should be compressed, false otherwise.
+	 */
+	bool ShouldCompressPackageData(uint64 PackageDataSize) const;
 
 	/** Root path to store all session data under */
 	FString SessionPath;
@@ -911,6 +927,12 @@ private:
 
 	/** Internal SQLite database */
 	TUniquePtr<FSQLiteDatabase> Database;
+
+	/** Threshold used to avoid caching very large packages (3GB for example) in memory. */
+	uint64 MaxPackageBlobSizeForCaching;
+
+	/** Threshold used to decide when a package should be compressed or not. */
+	uint64 MaxPackageDataSizeForCompression;
 };
 
 namespace ConcertSyncSessionDatabaseFilterUtil
diff --git a/Engine/Plugins/Developer/Concert/ConcertSync/ConcertSyncCore/Source/ConcertSyncCore/Public/ConcertSyncSessionFlags.h b/Engine/Plugins/Developer/Concert/ConcertSync/ConcertSyncCore/Source/ConcertSyncCore/Public/ConcertSyncSessionFlags.h
--- a/Engine/Plugins/Developer/Concert/ConcertSync/ConcertSyncCore/Source/ConcertSyncCore/Public/ConcertSyncSessionFlags.h
+++ b/Engine/Plugins/Developer/Concert/ConcertSync/ConcertSyncCore/Source/ConcertSyncCore/Public/ConcertSyncSessionFlags.h
@@ -27,22 +27,24 @@ enum class EConcertSyncSessionFlags : uint32
 	EnableSequencer = 1<<5,
 	/** Enable client presence tracking */
 	EnablePresence = 1<<6,
-	
+	/** Enable usage of file sharing to share large blobs (like package data) rather than embedding that data in message/event.*/
+	EnableFileSharing = 1<<7,
+
 	/**
 	 * Flags controlling behavior of enabled features
 	 */
 	
 	/** Discard transactions when a package is unloaded, eg) from a content hot-reload or closing a world editor (requires EnableTransactions) */
-	ShouldDiscardTransactionsOnPackageUnload = 1<<7,
+	ShouldDiscardTransactionsOnPackageUnload = 1<<8,
 	/** Send snapshots of multi-frame transactions (like drags) to other clients (requires EnableTransactions) */
-	ShouldSendTransactionSnapshots = 1<<8,
+	ShouldSendTransactionSnapshots = 1<<9,
 	/** Send the pristine state of a package (as revision 0) before saving it for the first time (requires EnablePackages) */
-	ShouldSendPackagePristineState = 1<<9,
+	ShouldSendPackagePristineState = 1<<10,
 	/** Send package auto-saves, in addition to user-saves (requires EnablePackages) */
-	ShouldSendPackageAutoSaves = 1<<10,
+	ShouldSendPackageAutoSaves = 1<<11,
 	/** Use a sandbox to hold package writes during a session (requires EnablePackages) */
-	ShouldUsePackageSandbox = 1<<11,
-	
+	ShouldUsePackageSandbox = 1<<12,
+
 	/**
 	 * Defaults for different scenarios
 	 */
@@ -50,6 +52,6 @@ enum class EConcertSyncSessionFlags : uint32
 	/** Default settings for a multi-user session */
 	Default_MultiUserSession = EnableLiveSync | EnableConnectionHistory | EnableLocking | EnableTransactions | EnablePackages | EnableSequencer | EnablePresence | ShouldSendTransactionSnapshots | ShouldUsePackageSandbox,
 	/** Default settings for a disaster recovery session */
-	Default_DisasterRecoverySession = EnableTransactions | EnablePackages | ShouldDiscardTransactionsOnPackageUnload | ShouldSendPackagePristineState | ShouldSendPackageAutoSaves,
+	Default_DisasterRecoverySession = EnableTransactions | EnablePackages | EnableFileSharing | ShouldDiscardTransactionsOnPackageUnload | ShouldSendPackagePristineState | ShouldSendPackageAutoSaves,
 };
 ENUM_CLASS_FLAGS(EConcertSyncSessionFlags);
diff --git a/Engine/Plugins/Developer/Concert/ConcertSync/ConcertSyncCore/Source/ConcertSyncCore/Public/ConcertSyncSessionTypes.cpp b/Engine/Plugins/Developer/Concert/ConcertSync/ConcertSyncCore/Source/ConcertSyncCore/Public/ConcertSyncSessionTypes.cpp
--- a/Engine/Plugins/Developer/Concert/ConcertSync/ConcertSyncCore/Source/ConcertSyncCore/Public/ConcertSyncSessionTypes.cpp
+++ b/Engine/Plugins/Developer/Concert/ConcertSync/ConcertSyncCore/Source/ConcertSyncCore/Public/ConcertSyncSessionTypes.cpp
@@ -734,13 +734,14 @@ FText FConcertSyncTransactionActivitySummary::CreateDisplayTextForUser(const FTe
 }
 
 
-FConcertSyncPackageActivitySummary FConcertSyncPackageActivitySummary::CreateSummaryForEvent(const FConcertSyncPackageEvent& InEvent)
+FConcertSyncPackageActivitySummary FConcertSyncPackageActivitySummary::CreateSummaryForEvent(const FConcertPackageInfo& PackageInfo)
 {
 	FConcertSyncPackageActivitySummary ActivitySummary;
-	ActivitySummary.PackageName = InEvent.Package.Info.PackageName;
-	ActivitySummary.NewPackageName = InEvent.Package.Info.NewPackageName;
-	ActivitySummary.PackageUpdateType = InEvent.Package.Info.PackageUpdateType;
-	ActivitySummary.bAutoSave = InEvent.Package.Info.bAutoSave;
+	ActivitySummary.PackageName = PackageInfo.PackageName;
+	ActivitySummary.NewPackageName = PackageInfo.NewPackageName;
+	ActivitySummary.PackageUpdateType = PackageInfo.PackageUpdateType;
+	ActivitySummary.bAutoSave = PackageInfo.bAutoSave;
+	ActivitySummary.bPreSave = PackageInfo.bPreSave;
 	return ActivitySummary;
 }
 
@@ -756,9 +757,18 @@ FText FConcertSyncPackageActivitySummary::CreateDisplayText(const bool InUseRich
 		FormatPattern = LOCTEXT("CreateDisplayText_Package_Added", "Added package {PackageName}.");
 		break;
 	case EConcertPackageUpdateType::Saved:
-		FormatPattern = bAutoSave 
-			? LOCTEXT("CreateDisplayText_Package_AutoSaved", "Auto-saved package {PackageName}.") 
-			: LOCTEXT("CreateDisplayText_Package_Saved", "Saved package {PackageName}.");
+		if (bPreSave)
+		{
+			FormatPattern = bAutoSave
+				? LOCTEXT("CreateDisplayText_Package_PreAutoSaved", "Captured package {PackageName} original state.")
+				: LOCTEXT("CreateDisplayText_Package_PreSaved", "Captured package {PackageName} original state.");
+		}
+		else
+		{
+			FormatPattern = bAutoSave
+				? LOCTEXT("CreateDisplayText_Package_AutoSaved", "Auto-saved package {PackageName}.") 
+				: LOCTEXT("CreateDisplayText_Package_Saved", "Saved package {PackageName}.");
+		}
 		break;
 	case EConcertPackageUpdateType::Renamed:
 		FormatPattern = LOCTEXT("CreateDisplayText_Package_Renamed", "Renamed package {PackageName} to {NewPackageName}.");
diff --git a/Engine/Plugins/Developer/Concert/ConcertSync/ConcertSyncCore/Source/ConcertSyncCore/Public/ConcertSyncSessionTypes.h b/Engine/Plugins/Developer/Concert/ConcertSync/ConcertSyncCore/Source/ConcertSyncCore/Public/ConcertSyncSessionTypes.h
--- a/Engine/Plugins/Developer/Concert/ConcertSync/ConcertSyncCore/Source/ConcertSyncCore/Public/ConcertSyncSessionTypes.h
+++ b/Engine/Plugins/Developer/Concert/ConcertSync/ConcertSyncCore/Source/ConcertSyncCore/Public/ConcertSyncSessionTypes.h
@@ -122,6 +122,44 @@ struct FConcertSyncPackageEvent
 	FConcertPackage Package;
 };
 
+/** Meta data for a package event in a Concert Sync Session. */
+USTRUCT()
+struct FConcertSyncPackageEventMetaData
+{
+	GENERATED_BODY()
+
+	/** The revision of this package within the session. */
+	UPROPERTY()
+	int64 PackageRevision = 0;
+
+	/** Contains information about the package event such as the package name, the event type, if this was triggered by an auto-save, etc. */
+	UPROPERTY()
+	FConcertPackageInfo PackageInfo;
+};
+
+/** Used to stream the package data. */
+struct FConcertPackageDataStream
+{
+	/** The package data, positioned to read the first byte of package data, not necessarily at zero, some unrelated data can be store before. Can be null. */
+	FArchive* DataAr = nullptr;
+
+	/** The size of the package data. Can 0 up to several GB large. Does't necessarily extend to the end of the archive, some unrelated data can be stored after.*/
+	int64 DataSize = 0;
+
+	/** An array of bytes containing the package data if the data was already available in memory, null otherwise. Can optimize few cases when available. */
+	const TArray<uint8>* DataBlob = nullptr;
+};
+
+/** Contains a package event where the package data is represented by a stream because it can be very large (several GB) */
+struct FConcertSyncPackageEventData
+{
+	/** The package event meta data.*/
+	FConcertSyncPackageEventMetaData MetaData;
+
+	/** The package data. */
+	FConcertPackageDataStream PackageDataStream;
+};
+
 /** Data for an activity entry in a Concert Sync Session */
 USTRUCT()
 struct FConcertSyncActivity
@@ -354,8 +392,12 @@ struct CONCERTSYNCCORE_API FConcertSyncPackageActivitySummary : public FConcertS
 	UPROPERTY()
 	bool bAutoSave = false;
 
+	/** Are we summarizing a pre-save update? */
+	UPROPERTY()
+	bool bPreSave = false;
+
 	/** Create this summary from a package event */
-	static FConcertSyncPackageActivitySummary CreateSummaryForEvent(const FConcertSyncPackageEvent& InEvent);
+	static FConcertSyncPackageActivitySummary CreateSummaryForEvent(const FConcertPackageInfo& InEvent);
 
 protected:
 	//~ FConcertSyncActivitySummary interface
diff --git a/Engine/Plugins/Developer/Concert/ConcertSync/ConcertSyncCore/Source/ConcertSyncCore/Public/ConcertWorkspaceData.h b/Engine/Plugins/Developer/Concert/ConcertSync/ConcertSyncCore/Source/ConcertSyncCore/Public/ConcertWorkspaceData.h
--- a/Engine/Plugins/Developer/Concert/ConcertSync/ConcertSyncCore/Source/ConcertSyncCore/Public/ConcertWorkspaceData.h
+++ b/Engine/Plugins/Developer/Concert/ConcertSync/ConcertSyncCore/Source/ConcertSyncCore/Public/ConcertWorkspaceData.h
@@ -59,9 +59,40 @@ struct FConcertPackage
 {
 	GENERATED_BODY()
 
+	/** Contains information about the package event such as the package name, the event type, if this was triggered by an auto-save, etc. */
 	UPROPERTY()
 	FConcertPackageInfo Info;
 
+	/** Contains the package data, unless the package size was too large (according to the hard limit or a configuration). */
 	UPROPERTY()
 	TArray<uint8> PackageData;
+
+	/** A link to a file containing the data if the package size was too large to be directly embedded. The package data needs to be retrieved using IConcertFileSharingService interface.*/
+	UPROPERTY()
+	FString FileId;
+
+	/** Whether some package data is embedded or linked. */
+	bool HasPackageData() const
+	{
+		return !FileId.IsEmpty() || PackageData.Num() != 0;
+	}
+
+	/** Whether the package data should be embedded in 'PackageData'. Larger packages could still be embedded as byte array if file sharing is not enabled. */
+	static bool ShouldEmbedPackageDataAsByteArray(uint64 PackageDataSize)
+	{
+		return PackageDataSize <= 2 * 1024 * 1024; // All package data below this size should be embedded in 'PackageData' byte array. This is a reasonable and efficient limit.
+	}
+
+	/** Whether the package data can safely be embedded in 'PackageData' without overflow (ex 3GB overflows int32 type). Not necessarily reasonable, but safe. */
+	static bool CanEmbedPackageDataAsByteArray(uint64 PackageDataSize)
+	{
+		return PackageDataSize <= GetMaxPackageDataSizeEmbeddableAsByteArray();
+	}
+
+	/** The maximum package data size that can be safely embedded in 'PackageData'. (Not necessarily reasonable, but safe). */
+	static constexpr uint64 GetMaxPackageDataSizeEmbeddableAsByteArray()
+	{
+		// The hard limit may not be safe by itself (nor reasonable) if the structure is serialized into another TArray<> with other data. Make it more reasonable, use only x% of the range.
+		return (static_cast<uint64>(TNumericLimits<decltype(PackageData)::SizeType>::Max()) * 8) / 10;
+	}
 };
diff --git a/Engine/Plugins/Developer/Concert/ConcertSync/ConcertSyncCore/Source/ConcertSyncCore/Public/IConcertFileSharingService.h b/Engine/Plugins/Developer/Concert/ConcertSync/ConcertSyncCore/Source/ConcertSyncCore/Public/IConcertFileSharingService.h
new file mode 100644
--- /dev/null
+++ b/Engine/Plugins/Developer/Concert/ConcertSync/ConcertSyncCore/Source/ConcertSyncCore/Public/IConcertFileSharingService.h
@@ -0,0 +1,64 @@
+// Copyright Epic Games, Inc. All Rights Reserved.
+
+#pragma once
+
+#include "CoreMinimal.h"
+
+/**
+ * Implements a file sharing service used as a side channel to ConcertTransport to exchange large files. This API is designed to
+ * work hand in hand with the request/response/event layer.
+ *
+ * When a file needs to be exchanged between the client and the server:
+ *
+ * The publisher side would:
+ *    - Publish a file using file sharing service. The service tracks the request and capture the file for sharing (can sent it to a server asynchronously, can be copied to a shared folder, etc.)
+ *    - Notify the consumer(s) by sending a Concert request/event/response with the ID of the published file.
+ *
+ * The consumer(s) side would:
+ *    - Receive and process the Concert request/event/response and extract the file ID.
+ *    - Create a stream to read the shared file, ideally in a background task. If required, new method could be added to simplify asynchrous IO.
+ *
+ * * @note The lifetime of the shared files is implementation specific.
+ *
+ */
+class IConcertFileSharingService
+{
+public:
+	virtual ~IConcertFileSharingService() = default;
+
+	/**
+	 * Copy the specified files in a uniquely named file stored in the service directory for sharing. The lifetime of the shared file
+	 * is implementation specific.
+	 * @param Pathname The file to copy and share.
+	 * @param OutFileId The ID of the published file. The publisher is expected to send the ID through the Concert transport layer to the consumer(s) (request/response/event).
+	 * @return True if the file was published successfully.
+	 */
+	virtual bool Publish(const FString& Pathname, FString& OutFileId) = 0;
+
+	/**
+	 * Copies up to Count bytes contained in the archives in a uniquely named file stored in the service directory for sharing.
+	 * The functions reads the archive from its current position up to \a Count bytes. The lifetime of the shared file
+	 * is implementation specific.
+	 * @param SrcAr The Archive containing the bytes to share. The archive is already at the right position, not necessarily 0.
+	 * @param Count The number of bytes to read from the archive.
+	 * @param OutFileId The ID of the published file. The publisher is expected to send the ID through the Concert transport layer to the consumer(s) (request/response/event).
+	 * @return True if the file was published successfully.
+	 */
+	virtual bool Publish(FArchive& SrcAr, int64 Count, FString& OutFileId) = 0;
+
+	/**
+	 * Copies all the bytes contained in the archives in a uniquely named file stored in the service directory for sharing.
+	 * @param SrcAr The Archive containing the bytes to share. Expected to be at position 0.
+	 * @param OutFileId The ID of the published file. The publisher is expected to send the ID through the Concert transport layer to the consumer(s) (request/response/event).
+	 * @return True if the file was published successfully.
+	 */
+	bool Publish(FArchive& SrcAr, FString& OutFileId) { return Publish(SrcAr, SrcAr.TotalSize(), OutFileId); }
+
+	/**
+	 * Open a shared file for consumption.
+	 * @param InFileId The Id of the file to read. The consumer is expected to receives this ID through the Concert transport layer (request/response/event).
+	 * @return An archive (loading mode) to stream the content of the file or nullptr if the file could not be found or opened.
+	 * @note The function can be used to poll until the resources is fully available. A more sophisticated API could be added to wait for a ressource if needed later.
+	 */
+	virtual TSharedPtr<FArchive> CreateReader(const FString& InFileId) = 0;
+};
diff --git a/Engine/Plugins/Developer/Concert/ConcertSync/ConcertSyncServer/Source/ConcertSyncServer/Private/ConcertServerWorkspace.cpp b/Engine/Plugins/Developer/Concert/ConcertSync/ConcertSyncServer/Source/ConcertSyncServer/Private/ConcertServerWorkspace.cpp
--- a/Engine/Plugins/Developer/Concert/ConcertSync/ConcertSyncServer/Source/ConcertSyncServer/Private/ConcertServerWorkspace.cpp
+++ b/Engine/Plugins/Developer/Concert/ConcertSync/ConcertSyncServer/Source/ConcertSyncServer/Private/ConcertServerWorkspace.cpp
@@ -3,6 +3,7 @@
 #include "ConcertServerWorkspace.h"
 
 #include "IConcertSession.h"
+#include "IConcertFileSharingService.h"
 #include "ConcertSyncServerLiveSession.h"
 #include "ConcertServerSyncCommandQueue.h"
 #include "ConcertSyncSessionDatabase.h"
@@ -12,7 +13,8 @@
 
 #include "Algo/Transform.h"
 
-FConcertServerWorkspace::FConcertServerWorkspace(const TSharedRef<FConcertSyncServerLiveSession>& InLiveSession)
+FConcertServerWorkspace::FConcertServerWorkspace(const TSharedRef<FConcertSyncServerLiveSession>& InLiveSession, TSharedPtr<IConcertFileSharingService> InFileSharingService)
+	: FileSharingService(MoveTemp(InFileSharingService))
 {
 	BindSession(InLiveSession);
 }
@@ -285,27 +287,50 @@ void FConcertServerWorkspace::HandlePackageUpdateEvent(const FConcertSessionCont
 	{
 		// If the client has the lock, then add the package activity
 		{
-			FConcertSyncPackageActivity PackageActivity;
-			PackageActivity.EndpointId = Context.SourceEndpointId;
-			PackageActivity.EventData.Package = Event.Package;
-			if (PackageActivity.EventData.Package.Info.PackageUpdateType == EConcertPackageUpdateType::Dummy && PackageActivity.EventData.Package.PackageData.Num() == 0)
+			bool bActivityAdded = false;
+
+			// Fill up required fields for the 'base' part of the activity.
+			FConcertSyncPackageActivity PackageActivityBasePart;
+			PackageActivityBasePart.EndpointId = Context.SourceEndpointId;
+			PackageActivityBasePart.EventSummary.SetTypedPayload(FConcertSyncPackageActivitySummary::CreateSummaryForEvent(Event.Package.Info));
+			PackageActivityBasePart.bIgnored = ShouldIgnoreClientActivityOnRestore(Context.SourceEndpointId);
+
+			// Fill up required fields for the 'event' part of the activity.
+			FConcertPackage Package;
+			Package.Info = Event.Package.Info;
+
+			if (Event.Package.Info.PackageUpdateType == EConcertPackageUpdateType::Dummy)
 			{
-				// If this is a dummy package, attempt to migrate the package data from the current head package revision so that newly synced clients will receive the correct package data
-				FConcertPackageInfo HeadPackageInfo;
-				if (LiveSession->GetSessionDatabase().GetPackageDataForRevision(PackageActivity.EventData.Package.Info.PackageName, HeadPackageInfo, &PackageActivity.EventData.Package.PackageData))
+				// Without Live Sync, the clients don't track the ongoing transactions in their 'local' database and can't provide the current transaction event ID when unsaved modifications to a package are discarded (to discard corresponding live transactions on this package).
+				if (!EnumHasAnyFlags(LiveSession->GetSessionFlags(), EConcertSyncSessionFlags::EnableLiveSync))
 				{
-					PackageActivity.EventData.Package.Info.PackageFileExtension = HeadPackageInfo.PackageFileExtension;
+					LiveSession->GetSessionDatabase().GetTransactionMaxEventId(Package.Info.TransactionEventIdAtSave); // Get the last transaction event from the server db.
 				}
 
-				// Without Live Sync, the clients don't track the ongoing transactions in their 'local' database and can't provide the current transaction event ID when unsaved modifications to a package are discarded (to discard corresponding live transactions on this package).
-				if (!EnumHasAnyFlags(LiveSession->GetSessionFlags(), EConcertSyncSessionFlags::EnableLiveSync))
+				if (!Event.Package.HasPackageData()) // The dummy event didn't come with any package data.
 				{
-					LiveSession->GetSessionDatabase().GetTransactionMaxEventId(PackageActivity.EventData.Package.Info.TransactionEventIdAtSave); // Get the last transaction event from the server db.
+					// Attempt to migrate the package data from the current head package revision so that newly synced clients will receive the correct package data.
+					LiveSession->GetSessionDatabase().GetPackageDataForRevision(Package.Info.PackageName, [this, &bActivityAdded, &Package, &PackageActivityBasePart](const FConcertPackageInfo& InHeadPackageInfo, FConcertPackageDataStream& InHeadPackageDataStream)
+					{
+						// Add the new package activity with the head revision data.
+						Package.Info.PackageFileExtension = InHeadPackageInfo.PackageFileExtension;
+						AddPackageActivity(PackageActivityBasePart, Package.Info, InHeadPackageDataStream);
+						bActivityAdded = true;
+					});
 				}
 			}
-			PackageActivity.EventSummary.SetTypedPayload(FConcertSyncPackageActivitySummary::CreateSummaryForEvent(PackageActivity.EventData));
-			PackageActivity.bIgnored = ShouldIgnoreClientActivityOnRestore(Context.SourceEndpointId);
-			AddPackageActivity(PackageActivity);
+
+			if (!bActivityAdded)
+			{
+				TSharedPtr<FArchive> PackageDataAr = GetPackageDataStream(Event.Package);
+				FConcertPackageDataStream PackageDataStream { 
+					PackageDataAr.Get(),
+					PackageDataAr ? PackageDataAr->TotalSize() : 0,
+					Event.Package.PackageData.Num() ? &Event.Package.PackageData : nullptr
+				};
+
+				AddPackageActivity(PackageActivityBasePart, Package.Info, PackageDataStream);
+			}
 		}
 
 		// Explicitly unlock the resource after saving it
@@ -445,16 +470,7 @@ EConcertSessionResponseCode FConcertServerWorkspace::HandleSyncEventRequest(cons
 			return EConcertSessionResponseCode::Success;
 		}
 	}
-	else if (Request.EventType == EConcertSyncActivityEventType::Package)
-	{
-		FConcertSyncPackageEvent PackageEvent;
-		if (LiveSession->GetSessionDatabase().GetPackageEvent(Request.EventId, PackageEvent, /*bMetaDataOnly*/false)) // The request is meant to retrive the package data from a partially sync event.
-		{
-			Response.Event.SetTypedPayload(PackageEvent);
-			return EConcertSessionResponseCode::Success;
-		}
-	}
-	// else Other event types (connection/lock) don't have interesting non-meta-data data to retrieve.
+	// else Other event types (package/connection/lock) don't have interesting non-meta-data data to retrieve.
 
 	return EConcertSessionResponseCode::Failed;
 }
@@ -494,14 +510,16 @@ void FConcertServerWorkspace::HandleEndPlaySession(const FName InPlayPackageName
 	{
 		// Save a dummy package to discard the live transactions for the previous play world
 		// Play worlds are never saved, so we don't have to worry about migrating over the previous data here
-		FConcertSyncPackageActivity DummyPackageActivity;
+		FConcertSyncActivity DummyPackageActivity;
+		FConcertPackageInfo DummyPackageInfo;
+		FConcertPackageDataStream DummyPackageStream;
 		DummyPackageActivity.EndpointId = InEndpointId;
-		DummyPackageActivity.EventData.Package.Info.PackageUpdateType = EConcertPackageUpdateType::Dummy;
-		DummyPackageActivity.EventData.Package.Info.PackageName = InPlayPackageName;
-		LiveSession->GetSessionDatabase().GetTransactionMaxEventId(DummyPackageActivity.EventData.Package.Info.TransactionEventIdAtSave);
-		DummyPackageActivity.EventSummary.SetTypedPayload(FConcertSyncPackageActivitySummary::CreateSummaryForEvent(DummyPackageActivity.EventData));
+		DummyPackageInfo.PackageUpdateType = EConcertPackageUpdateType::Dummy;
+		DummyPackageInfo.PackageName = InPlayPackageName;
+		LiveSession->GetSessionDatabase().GetTransactionMaxEventId(DummyPackageInfo.TransactionEventIdAtSave);
+		DummyPackageActivity.EventSummary.SetTypedPayload(FConcertSyncPackageActivitySummary::CreateSummaryForEvent(DummyPackageInfo));
 		DummyPackageActivity.bIgnored = ShouldIgnoreClientActivityOnRestore(InEndpointId);
-		AddPackageActivity(DummyPackageActivity);
+		AddPackageActivity(DummyPackageActivity, DummyPackageInfo, DummyPackageStream);
 	}
 }
 
@@ -909,12 +927,13 @@ void FConcertServerWorkspace::SendSyncTransactionActivityEvent(const FGuid& InTa
 	}
 }
 
-void FConcertServerWorkspace::AddPackageActivity(const FConcertSyncPackageActivity& InPackageActivity)
+void FConcertServerWorkspace::AddPackageActivity(const FConcertSyncActivity& InPackageActivityBasePart, const FConcertPackageInfo& PackageInfo, FConcertPackageDataStream& InPackageDataStream)
 {
 	// Add the activity and sync it
 	int64 ActivityId = 0;
 	int64 EventId = 0;
-	if (LiveSession->GetSessionDatabase().AddPackageActivity(InPackageActivity, ActivityId, EventId))
+
+	if (LiveSession->GetSessionDatabase().AddPackageActivity(InPackageActivityBasePart, PackageInfo, InPackageDataStream, ActivityId, EventId))
 	{
 		PostActivityAdded(ActivityId);
 		SyncCommandQueue->QueueCommand(LiveSyncEndpoints, [this, SyncActivityId = ActivityId](const FConcertServerSyncCommandQueue::FSyncCommandContext& InSyncCommandContext, const FGuid& InEndpointId)
@@ -940,7 +959,22 @@ void FConcertServerWorkspace::SendSyncPackageActivityEvent(const FGuid& InTarget
 			bMetaDataOnly = !ConcertSyncSessionDatabaseFilterUtil::PackageEventPassesFilter(SyncActivity.EventId, SessionFilter, LiveSession->GetSessionDatabase());
 		}
 
-		if (!LiveSession->GetSessionDatabase().GetPackageEvent(SyncActivity.EventId, SyncActivity.EventData, bMetaDataOnly))
+		if (!bMetaDataOnly) // Package data is needed.
+		{
+			// Callback function used to fill the package event (containing the package data, which can be small or large).
+			EConcertSessionResponseCode FillPackageEventResponseCode = EConcertSessionResponseCode::Success;
+			auto FillPackageEventFn = [this, &SyncActivity, &FillPackageEventResponseCode](FConcertSyncPackageEventData& PackageEventSrc)
+			{
+				// May embed the package data directly or link it if its too big and file sharing is enabled.
+				FillPackageEventResponseCode = FillPackageEvent(PackageEventSrc, SyncActivity.EventData);
+			};
+
+			if (!LiveSession->GetSessionDatabase().GetPackageEvent(SyncActivity.EventId, FillPackageEventFn) || FillPackageEventResponseCode == EConcertSessionResponseCode::Failed)
+			{
+				UE_LOG(LogConcert, Error, TEXT("Failed to get package event '%s' from live session '%s': %s"), *LexToString(SyncActivity.EventId), *LiveSession->GetSession().GetName(), *LiveSession->GetSessionDatabase().GetLastError());
+			}
+		}
+		else if (!LiveSession->GetSessionDatabase().GetPackageEventMetaData(SyncActivity.EventId, SyncActivity.EventData.PackageRevision, SyncActivity.EventData.Package.Info)) // Ignore package data.
 		{
 			UE_LOG(LogConcert, Error, TEXT("Failed to get package event '%s' from live session '%s': %s"), *LexToString(SyncActivity.EventId), *LiveSession->GetSession().GetName(), *LiveSession->GetSessionDatabase().GetLastError());
 		}
@@ -958,7 +992,7 @@ void FConcertServerWorkspace::SendSyncPackageActivityEvent(const FGuid& InTarget
 
 void FConcertServerWorkspace::PostActivityAdded(const int64 InActivityId)
 {
-	FConcertSyncPackageActivity Activity;
+	FConcertSyncActivity Activity;
 	if (LiveSession->GetSessionDatabase().GetActivity(InActivityId, Activity))
 	{
 		FConcertSyncEndpointData EndpointData;
@@ -974,3 +1008,50 @@ void FConcertServerWorkspace::PostActivityAdded(const int64 InActivityId)
 		}
 	}
 }
+
+bool FConcertServerWorkspace::CanExchangePackageDataAsByteArray(int64 PackageDataSize) const
+{
+	if (FileSharingService && EnumHasAnyFlags(LiveSession->GetSessionFlags(), EConcertSyncSessionFlags::EnableFileSharing))
+	{
+		return FConcertPackage::ShouldEmbedPackageDataAsByteArray(PackageDataSize); // Test the package data size against a preferred limit.
+	}
+
+	return FConcertPackage::CanEmbedPackageDataAsByteArray(PackageDataSize); // The the package data size against the maximum permitted.
+}
+
+EConcertSessionResponseCode FConcertServerWorkspace::FillPackageEvent(FConcertSyncPackageEventData& InPackageEvent, FConcertSyncPackageEvent& OutPackageEvent) const
+{
+	EConcertSessionResponseCode FillPackageEventResponseCode = EConcertSessionResponseCode::Success;
+
+	// Fill the meta data.
+	OutPackageEvent.PackageRevision = InPackageEvent.MetaData.PackageRevision;
+	OutPackageEvent.Package.Info = InPackageEvent.MetaData.PackageInfo;
+
+	// Fill the package data.
+	if (CanExchangePackageDataAsByteArray(InPackageEvent.PackageDataStream.DataSize))
+	{
+		// Store the package data directly in the package event.
+		OutPackageEvent.Package.PackageData.AddUninitialized(InPackageEvent.PackageDataStream.DataSize);
+		InPackageEvent.PackageDataStream.DataAr->Serialize(OutPackageEvent.Package.PackageData.GetData(), InPackageEvent.PackageDataStream.DataSize);
+	}
+	else if (FileSharingService)
+	{
+		// Share a link to the file.
+		FillPackageEventResponseCode = FileSharingService->Publish(*InPackageEvent.PackageDataStream.DataAr, InPackageEvent.PackageDataStream.DataSize, OutPackageEvent.Package.FileId) ? EConcertSessionResponseCode::Success : EConcertSessionResponseCode::Failed;
+	}
+	else
+	{
+		// The package data is too big.
+		FillPackageEventResponseCode = EConcertSessionResponseCode::Failed;
+	}
+
+	return FillPackageEventResponseCode;
+}
+
+TSharedPtr<FArchive> FConcertServerWorkspace::GetPackageDataStream(const FConcertPackage& Package)
+{
+	return !Package.FileId.IsEmpty() ?
+		FileSharingService->CreateReader(Package.FileId) : // Package data is not embedded with the activty itself, but linked.
+		MakeShared<FMemoryReader>(Package.PackageData); // Package data is embedded in the activity object.
+}
+
diff --git a/Engine/Plugins/Developer/Concert/ConcertSync/ConcertSyncServer/Source/ConcertSyncServer/Private/ConcertServerWorkspace.h b/Engine/Plugins/Developer/Concert/ConcertSync/ConcertSyncServer/Source/ConcertSyncServer/Private/ConcertServerWorkspace.h
--- a/Engine/Plugins/Developer/Concert/ConcertSync/ConcertSyncServer/Source/ConcertSyncServer/Private/ConcertServerWorkspace.h
+++ b/Engine/Plugins/Developer/Concert/ConcertSync/ConcertSyncServer/Source/ConcertSyncServer/Private/ConcertServerWorkspace.h
@@ -7,6 +7,7 @@
 #include "ConcertWorkspaceMessages.h"
 
 class IConcertServerSession;
+class IConcertFileSharingService;
 class FConcertSyncServerLiveSession;
 class FConcertServerSyncCommandQueue;
 class FConcertServerDataStore;
@@ -25,7 +26,7 @@ ENUM_CLASS_FLAGS(EConcertLockFlags);
 class FConcertServerWorkspace
 {
 public:
-	explicit FConcertServerWorkspace(const TSharedRef<FConcertSyncServerLiveSession>& InLiveSession);
+	explicit FConcertServerWorkspace(const TSharedRef<FConcertSyncServerLiveSession>& InLiveSession, TSharedPtr<IConcertFileSharingService> InFileSharingService);
 	~FConcertServerWorkspace();
 
 private:
@@ -202,7 +203,7 @@ private:
 	 *
 	 * @param InPackageActivity			The package activity to add (the ActivityId, EventTime, EventType, EventId, and PackageRevision members are ignored).
 	 */
-	void AddPackageActivity(const FConcertSyncPackageActivity& InPackageActivity);
+	void AddPackageActivity(const FConcertSyncActivity& InPackageActivityBasePart, const FConcertPackageInfo& PackageInfo, FConcertPackageDataStream& InPackageDataStream);
 
 	/**
 	 * Send a sync event for a package activity in the session database.
@@ -228,6 +229,23 @@ private:
 	 */
 	bool ShouldIgnoreClientActivityOnRestore(const FGuid ClientEndpoint) const { return IgnoredActivityClients.Contains(ClientEndpoint); }
 
+	/**
+	 * Returns true if the package data is small enough to be exchanged using a single TArray<> data structure. If the package data is
+	 * too big according to a configuration (or a hard limit) the package data will be exchanged through other means like file stream
+	 * or file sharing.
+	 */
+	bool CanExchangePackageDataAsByteArray(int64 PackageDataSize) const;
+
+	/**
+	 * Write all package event fields from the event and decide if the package data should be directly embedded in the event or linked to a shared file.
+	 */
+	EConcertSessionResponseCode FillPackageEvent(FConcertSyncPackageEventData& InPackageEventSrc, FConcertSyncPackageEvent& OutPackageEvent) const;
+
+	/**
+	 * Open a stream on the package data, depending if the package data was directly embedded in the object (in memory) or linked as a shared file.
+	 */
+	TSharedPtr<FArchive> GetPackageDataStream(const FConcertPackage& Package);
+
 	/** Live session tracked by this workspace */
 	TSharedPtr<FConcertSyncServerLiveSession> LiveSession;
 
@@ -266,4 +284,7 @@ private:
 
 	/** Contains the list of clients for which activities should be marked as 'ignore on restore'. @see FConcertIgnoreActivityStateChangedEvent. */
 	TSet<FGuid> IgnoredActivityClients;
+
+	/** Optional side channel to exchange large blobs (package data) with the client in a scalable way (ex. the request/response transport layer is not designed and doesn't support exchanging 3GB packages). */
+	TSharedPtr<IConcertFileSharingService> FileSharingService;
 };
diff --git a/Engine/Plugins/Developer/Concert/ConcertSync/ConcertSyncServer/Source/ConcertSyncServer/Private/ConcertSyncServer.cpp b/Engine/Plugins/Developer/Concert/ConcertSync/ConcertSyncServer/Source/ConcertSyncServer/Private/ConcertSyncServer.cpp
--- a/Engine/Plugins/Developer/Concert/ConcertSync/ConcertSyncServer/Source/ConcertSyncServer/Private/ConcertSyncServer.cpp
+++ b/Engine/Plugins/Developer/Concert/ConcertSync/ConcertSyncServer/Source/ConcertSyncServer/Private/ConcertSyncServer.cpp
@@ -167,20 +167,43 @@ bool MigrateSessionData(const FConcertSyncSessionDatabase& InSourceDatabase, con
 
 			case EConcertSyncActivityEventType::Package:
 			{
-				FConcertSyncPackageActivity PackageActivity;
-				if (!InSourceDatabase.GetActivity(InActivityId, PackageActivity))
+				FConcertSyncActivity PackageActivityBasePart;
+				if (!InSourceDatabase.GetActivity(InActivityId, PackageActivityBasePart)) // This only fill up the part that is common to all activities.
 				{
 					MIGRATE_SET_ERROR_RESULT_AND_RETURN("Failed to get package activity '%s' from database at '%s': %s", *LexToString(InActivityId), *InSourceDatabase.GetFilename(), *InSourceDatabase.GetLastError());
 				}
-				if ((InDestSessionFilter.bIncludeIgnoredActivities || !PackageActivity.bIgnored) && ConcertSyncSessionDatabaseFilterUtil::PackageEventPassesFilter(PackageActivity.EventId, InDestSessionFilter, InSourceDatabase))
+				if ((InDestSessionFilter.bIncludeIgnoredActivities || !PackageActivityBasePart.bIgnored) && ConcertSyncSessionDatabaseFilterUtil::PackageEventPassesFilter(PackageActivityBasePart.EventId, InDestSessionFilter, InSourceDatabase))
 				{
-					if (!InSourceDatabase.GetPackageEvent(PackageActivity.EventId, PackageActivity.EventData, InDestSessionFilter.bMetaDataOnly))
+					if (InDestSessionFilter.bMetaDataOnly)
 					{
-						MIGRATE_SET_ERROR_RESULT_AND_RETURN("Failed to get package event '%s' from database at '%s': %s", *LexToString(PackageActivity.EventId), *InSourceDatabase.GetFilename(), *InSourceDatabase.GetLastError());
+						FConcertSyncPackageEventData PackageActivityEventPart;
+						if (!InSourceDatabase.GetPackageEventMetaData(PackageActivityBasePart.EventId, PackageActivityEventPart.MetaData.PackageRevision, PackageActivityEventPart.MetaData.PackageInfo))
+						{
+							check(PackageActivityEventPart.PackageDataStream.DataAr == nullptr);
+							MIGRATE_SET_ERROR_RESULT_AND_RETURN("Failed to get package event '%s' from database at '%s': %s", *LexToString(PackageActivityBasePart.EventId), *InSourceDatabase.GetFilename(), *InSourceDatabase.GetLastError());
+						}
+						// Merge the base part with the event part to reconstruct the package activity.
+						if (!DestDatabase.SetPackageActivity(PackageActivityBasePart, PackageActivityEventPart, InDestSessionFilter.bMetaDataOnly))
+						{
+							MIGRATE_SET_ERROR_RESULT_AND_RETURN("Failed to set package activity '%s' on database at '%s': %s", *LexToString(InActivityId), *DestDatabase.GetFilename(), *DestDatabase.GetLastError());
+						}
 					}
-					if (!DestDatabase.SetPackageActivity(PackageActivity, InDestSessionFilter.bMetaDataOnly))
+					else
 					{
-						MIGRATE_SET_ERROR_RESULT_AND_RETURN("Failed to set package activity '%s' on database at '%s': %s", *LexToString(InActivityId), *DestDatabase.GetFilename(), *DestDatabase.GetLastError());
+						// Pull the package event required to fill the package specific part of the activity.
+						bool bSuccess = InSourceDatabase.GetPackageEvent(PackageActivityBasePart.EventId, [&DestDatabase, &PackageActivityBasePart, &InDestSessionFilter, InActivityId](FConcertSyncPackageEventData& PackageActivityEventPart)
+						{
+							// Merge the common activity part with the package event to reconstruct and migrate a package activity.
+							if (!DestDatabase.SetPackageActivity(PackageActivityBasePart, PackageActivityEventPart, InDestSessionFilter.bMetaDataOnly))
+							{
+								UE_LOG(LogConcert, Error, TEXT("Failed to set package activity '%s' on database at '%s': %s"), *LexToString(InActivityId), *DestDatabase.GetFilename(), *DestDatabase.GetLastError());
+							}
+						});
+
+						if (!bSuccess)
+						{
+							MIGRATE_SET_ERROR_RESULT_AND_RETURN("Failed to get package event '%s' from database at '%s': %s", *LexToString(PackageActivityBasePart.EventId), *InSourceDatabase.GetFilename(), *InSourceDatabase.GetLastError());
+						}
 					}
 				}
 			}
@@ -232,6 +255,11 @@ IConcertServerRef FConcertSyncServer::GetConcertServer() const
 	return ConcertServer;
 }
 
+void FConcertSyncServer::SetFileSharingService(TSharedPtr<IConcertFileSharingService> InFileSharingService)
+{
+	FileSharingService = MoveTemp(InFileSharingService);
+}
+
 void FConcertSyncServer::GetSessionsFromPath(const IConcertServer& InServer, const FString& InPath, TArray<FConcertSessionInfo>& OutSessionInfos, TArray<FDateTime>* OutSessionCreationTimes)
 {
 	IFileManager::Get().IterateDirectory(*InPath, [&OutSessionInfos, &OutSessionCreationTimes](const TCHAR* FilenameOrDirectory, bool bIsDirectory) -> bool
@@ -403,7 +431,8 @@ bool FConcertSyncServer::GetSessionActivities(const FConcertSyncSessionDatabase&
 				// If the details are requested, get the package event meta-data, which contain extra info about the package name/version, etc.
 				if (bIncludeDetails)
 				{
-					Database.GetPackageEvent(SyncActivity.EventId, SyncActivity.EventData, true/*bMetaDataOnly*/); // Just get the package event meta-data, the package data is not required to display details.
+					// Don't pull the package data, it can be huge (few GB) and is not useful in the context of this function design. (It is used to display extra info in UI and parsing the large package data is not practical)
+					Database.GetPackageEventMetaData(SyncActivity.EventId, SyncActivity.EventData.PackageRevision, SyncActivity.EventData.Package.Info);
 				}
 
 				SerializedSyncActivityPayload.SetTypedPayload(SyncActivity);
@@ -441,7 +470,7 @@ void FConcertSyncServer::CreateWorkspace(const TSharedRef<FConcertSyncServerLive
 {
 	check(InLiveSession->IsValidSession());
 	DestroyWorkspace(InLiveSession);
-	LiveSessionWorkspaces.Add(InLiveSession->GetSession().GetId(), MakeShared<FConcertServerWorkspace>(InLiveSession));
+	LiveSessionWorkspaces.Add(InLiveSession->GetSession().GetId(), MakeShared<FConcertServerWorkspace>(InLiveSession, FileSharingService));
 }
 
 void FConcertSyncServer::DestroyWorkspace(const TSharedRef<FConcertSyncServerLiveSession>& InLiveSession)
diff --git a/Engine/Plugins/Developer/Concert/ConcertSync/ConcertSyncServer/Source/ConcertSyncServer/Private/ConcertSyncServer.h b/Engine/Plugins/Developer/Concert/ConcertSync/ConcertSyncServer/Source/ConcertSyncServer/Private/ConcertSyncServer.h
--- a/Engine/Plugins/Developer/Concert/ConcertSync/ConcertSyncServer/Source/ConcertSyncServer/Private/ConcertSyncServer.h
+++ b/Engine/Plugins/Developer/Concert/ConcertSync/ConcertSyncServer/Source/ConcertSyncServer/Private/ConcertSyncServer.h
@@ -27,6 +27,7 @@ public:
 	virtual void Startup(const UConcertServerConfig* InServerConfig, const EConcertSyncSessionFlags InSessionFlags) override;
 	virtual void Shutdown() override;
 	virtual IConcertServerRef GetConcertServer() const override;
+	virtual void SetFileSharingService(TSharedPtr<IConcertFileSharingService> InFileSharingService) override;
 
 	//~ IConcertServerEventSink interface
 	virtual void GetSessionsFromPath(const IConcertServer& InServer, const FString& InPath, TArray<FConcertSessionInfo>& OutSessionInfos, TArray<FDateTime>* OutSessionCreationTimes = nullptr) override;
@@ -75,4 +76,7 @@ private:
 
 	/** Map of archived session IDs to their associated session data */
 	TMap<FGuid, TSharedPtr<FConcertSyncServerArchivedSession>> ArchivedSessions;
+
+	/** Optional side channel to exchange large blobs (package data) with the server in a scalable way (ex. the request/response transport layer is not designed and doesn't support exchanging 3GB packages). */
+	TSharedPtr<IConcertFileSharingService> FileSharingService;
 };
diff --git a/Engine/Plugins/Developer/Concert/ConcertSync/ConcertSyncServer/Source/ConcertSyncServer/Public/ConcertSyncServerLoop.h b/Engine/Plugins/Developer/Concert/ConcertSync/ConcertSyncServer/Source/ConcertSyncServer/Public/ConcertSyncServerLoop.h
--- a/Engine/Plugins/Developer/Concert/ConcertSync/ConcertSyncServer/Source/ConcertSyncServer/Public/ConcertSyncServerLoop.h
+++ b/Engine/Plugins/Developer/Concert/ConcertSync/ConcertSyncServer/Source/ConcertSyncServer/Public/ConcertSyncServerLoop.h
@@ -7,6 +7,7 @@
 #include "ConcertSyncSessionFlags.h"
 
 class UConcertServerConfig;
+class IConcertFileSharingService;
 
 struct FConcertSyncServerLoopInitArgs
 {
@@ -25,11 +26,14 @@ struct FConcertSyncServerLoopInitArgs
 	/** The session filter to apply when auto-archiving sessions on this server */
 	FConcertSessionFilter ServiceAutoArchiveSessionFilter;
 
+	/** The optional file sharing server to exchange large files. Can be null. */
+	TSharedPtr<IConcertFileSharingService> FileSharingService;
+
 	/** Function to get the server settings object to configure the server with with, or unbound to parse the default settings */
 	TFunction<const UConcertServerConfig*()> GetServerConfigFunc;
-    
-    /** Whether the service should show the log console. */
-    bool bShowConsole = true;
+
+	/** Whether the service should show the log console. */
+	bool bShowConsole = true;
 };
 
 /**
diff --git a/Engine/Plugins/Developer/Concert/ConcertSync/ConcertSyncServer/Source/ConcertSyncServer/Public/ConcertSyncServerLoop.inl b/Engine/Plugins/Developer/Concert/ConcertSync/ConcertSyncServer/Source/ConcertSyncServer/Public/ConcertSyncServerLoop.inl
--- a/Engine/Plugins/Developer/Concert/ConcertSync/ConcertSyncServer/Source/ConcertSyncServer/Public/ConcertSyncServerLoop.inl
+++ b/Engine/Plugins/Developer/Concert/ConcertSync/ConcertSyncServer/Source/ConcertSyncServer/Public/ConcertSyncServerLoop.inl
@@ -86,6 +86,7 @@ int32 ConcertSyncServerLoop(int32 ArgC, TCHAR** ArgV, const FConcertSyncServerLo
 
 		// Setup Concert Sync to run in server mode.
 		TSharedPtr<IConcertSyncServer> ConcertSyncServer = IConcertSyncServerModule::Get().CreateServer(InitArgs.ServiceRole, InitArgs.ServiceAutoArchiveSessionFilter);
+		ConcertSyncServer->SetFileSharingService(InitArgs.FileSharingService);
 		ConcertSyncServer->Startup(ServerConfig, InitArgs.SessionFlags);
 
 		// if we have a default session, set it up properly
diff --git a/Engine/Plugins/Developer/Concert/ConcertSync/ConcertSyncServer/Source/ConcertSyncServer/Public/IConcertSyncServer.h b/Engine/Plugins/Developer/Concert/ConcertSync/ConcertSyncServer/Source/ConcertSyncServer/Public/IConcertSyncServer.h
--- a/Engine/Plugins/Developer/Concert/ConcertSync/ConcertSyncServer/Source/ConcertSyncServer/Public/IConcertSyncServer.h
+++ b/Engine/Plugins/Developer/Concert/ConcertSync/ConcertSyncServer/Source/ConcertSyncServer/Public/IConcertSyncServer.h
@@ -5,6 +5,7 @@
 #include "CoreMinimal.h"
 #include "IConcertModule.h"
 #include "ConcertSyncSessionFlags.h"
+class IConcertFileSharingService;
 
 class UConcertServerConfig;
 
@@ -24,4 +25,7 @@ public:
 
 	/** Get the current server */
 	virtual IConcertServerRef GetConcertServer() const = 0;
+
+	/** Set the file sharing service, enabling the server to work with large files. The server sharing service must be compatible with the client one. */
+	virtual void SetFileSharingService(TSharedPtr<IConcertFileSharingService> InFileSharingService) = 0;
 };
diff --git a/Engine/Plugins/Developer/Concert/ConcertSync/ConcertSyncTest/Source/ConcertSyncTest/Private/ConcertDatabaseTests.cpp b/Engine/Plugins/Developer/Concert/ConcertSync/ConcertSyncTest/Source/ConcertSyncTest/Private/ConcertDatabaseTests.cpp
--- a/Engine/Plugins/Developer/Concert/ConcertSync/ConcertSyncTest/Source/ConcertSyncTest/Private/ConcertDatabaseTests.cpp
+++ b/Engine/Plugins/Developer/Concert/ConcertSync/ConcertSyncTest/Source/ConcertSyncTest/Private/ConcertDatabaseTests.cpp
@@ -75,11 +75,13 @@ bool FConcertSessionDatabaseTest::RunTest(const FString& Parameters)
 				AddError(FString::Printf(TEXT("Failed to add transaction activity to server database: %s"), *SessionDatabase_Server.GetLastError()));
 			}
 
-			FConcertSyncPackageActivity PackageActivity;
-			PackageActivity.EndpointId = EndpointId;
-			PackageActivity.EventData.Package.Info.PackageName = TEXT("/Game/TestAsset");
-			SessionDatabase_Server.GetTransactionMaxEventId(PackageActivity.EventData.Package.Info.TransactionEventIdAtSave);
-			if (!SessionDatabase_Server.AddPackageActivity(PackageActivity, ActivityId, EventId))
+			FConcertSyncActivity PackageActivityBasePart;
+			FConcertPackageInfo PackageInfo;
+			FConcertPackageDataStream PackageDataStream;
+			PackageActivityBasePart.EndpointId = EndpointId;
+			PackageInfo.PackageName = TEXT("/Game/TestAsset");
+			SessionDatabase_Server.GetTransactionMaxEventId(PackageInfo.TransactionEventIdAtSave);
+			if (!SessionDatabase_Server.AddPackageActivity(PackageActivityBasePart, PackageInfo, PackageDataStream, ActivityId, EventId))
 			{
 				AddError(FString::Printf(TEXT("Failed to add package activity to server database: %s"), *SessionDatabase_Server.GetLastError()));
 			}
@@ -169,13 +171,18 @@ bool FConcertSessionDatabaseTest::RunTest(const FString& Parameters)
 
 			case EConcertSyncActivityEventType::Package:
 			{
-				FConcertSyncPackageActivity PackageActivity;
-				if (!SessionDatabase_Server.GetPackageActivity(InActivityId, PackageActivity))
+				bool bSetPackageActivitySucceeded = true;
+				bool bGetPackageActivitySucceeded = SessionDatabase_Server.GetPackageActivity(InActivityId, [&](FConcertSyncActivity&& ActivityBasePart, FConcertSyncPackageEventData& ActivityEventPart)
+				{
+					bSetPackageActivitySucceeded = SessionDatabase_Client.SetPackageActivity(ActivityBasePart, ActivityEventPart);
+				});
+
+				if (!bGetPackageActivitySucceeded)
 				{
 					AddError(FString::Printf(TEXT("Failed to get package activity '%s' from server database: %s"), *LexToString(InActivityId), *SessionDatabase_Server.GetLastError()));
 					return false;
 				}
-				if (!SessionDatabase_Client.SetPackageActivity(PackageActivity))
+				if (!bSetPackageActivitySucceeded)
 				{
 					AddError(FString::Printf(TEXT("Failed to set package activity '%s' on client database: %s"), *LexToString(InActivityId), *SessionDatabase_Client.GetLastError()));
 					return false;
diff --git a/Engine/Plugins/Developer/Concert/DisasterRecoveryClient/Source/DisasterRecoveryClient/Private/DisasterRecoveryClientModule.cpp b/Engine/Plugins/Developer/Concert/DisasterRecoveryClient/Source/DisasterRecoveryClient/Private/DisasterRecoveryClientModule.cpp
--- a/Engine/Plugins/Developer/Concert/DisasterRecoveryClient/Source/DisasterRecoveryClient/Private/DisasterRecoveryClientModule.cpp
+++ b/Engine/Plugins/Developer/Concert/DisasterRecoveryClient/Source/DisasterRecoveryClient/Private/DisasterRecoveryClientModule.cpp
@@ -10,6 +10,7 @@
 #include "IConcertSyncClient.h"
 #include "ConcertFrontendStyle.h"
 #include "ConcertUtil.h"
+#include "ConcertLocalFileSharingService.h"
 
 #include "Misc/Paths.h"
 #include "Misc/CoreDelegates.h"
@@ -105,6 +106,7 @@ public:
 
 		// Create and start the client.
 		DisasterRecoveryClient = IConcertSyncClientModule::Get().CreateClient(Role);
+		DisasterRecoveryClient->SetFileSharingService(MakeShared<FConcertLocalFileSharingService>(Role));
 		DisasterRecoveryClient->Startup(ClientConfig, EConcertSyncSessionFlags::Default_DisasterRecoverySession);
 
 		// Hook to sync session is created/destroyed event to detect when a Multi-User session (incompatible) is started.
diff --git a/Engine/Plugins/Developer/Concert/MultiUserClient/Source/MultiUserClient/Private/MultiUserClientModule.cpp b/Engine/Plugins/Developer/Concert/MultiUserClient/Source/MultiUserClient/Private/MultiUserClientModule.cpp
--- a/Engine/Plugins/Developer/Concert/MultiUserClient/Source/MultiUserClient/Private/MultiUserClientModule.cpp
+++ b/Engine/Plugins/Developer/Concert/MultiUserClient/Source/MultiUserClient/Private/MultiUserClientModule.cpp
@@ -10,7 +10,9 @@
 #include "ConcertLogGlobal.h"
 #include "ConcertFrontendStyle.h"
 #include "ConcertSyncSettings.h"
+#include "ConcertWorkspaceData.h"
 #include "ConcertWorkspaceUI.h"
+#include "ConcertUtil.h"
 #include "MultiUserClientUtils.h"
 
 #include "Misc/App.h"
@@ -594,6 +596,10 @@ public:
 			ClientConfig->bAutoConnect = false; // Prevent user from auto-connecting if there is no compatible plugin loaded to find and communicate with the server.
 		}
 
+		// Handle notification when a session is staring or ending.
+		MultiUserClient->OnWorkspaceStartup().AddRaw(this, &FMultiUserClientModule::OnWorkspaceStartup);
+		MultiUserClient->OnWorkspaceShutdown().AddRaw(this, &FMultiUserClientModule::OnWorkspaceShutdown);
+
 		// Boot the client instance
 		MultiUserClient->Startup(ClientConfig, EConcertSyncSessionFlags::Default_MultiUserSession);
 
@@ -733,6 +739,36 @@ public:
 	}
 
 private:
+	void OnWorkspaceStartup(const TSharedPtr<IConcertClientWorkspace>& InWorkspace)
+	{
+		//InWorkspace.OnPackageTooLargeError().AddRaw(this, &FMultiUserClientModule::HandlePackageTooLargeError);
+	}
+
+	void OnWorkspaceShutdown(const TSharedPtr<IConcertClientWorkspace>& InWorkspace)
+	{
+		//InWorkspace.OnPackageTooLargeError().RemoveAll(this);
+	}
+
+	void HandlePackageTooLargeError(const FConcertPackageInfo& PackageInfo, int64 PackageSize, int64 MaxPackageSizeSupported)
+	{
+		if (PackageInfo.PackageUpdateType == EConcertPackageUpdateType::Added)
+		{
+			UE_LOG(LogConcert, Warning, TEXT("The newly added package '%s' was too large (%s) to be synchronized on other Multi-User clients. It is recommended to save the package, leave and persist the session, submit the file to the source control, sync all clients to this submit and create a new session."), *PackageInfo.PackageName.ToString(), *FText::AsMemory(PackageSize).ToString());
+		}
+		else if (PackageInfo.PackageUpdateType == EConcertPackageUpdateType::Saved)
+		{
+			UE_LOG(LogConcert, Warning, TEXT("The saved package '%s' was too large (%s) to be synchronized on other Multi-User clients. It is recommended to leave and persist the session, submit the file to the source control, sync all clients to this submit and create a new session."), *PackageInfo.PackageName.ToString(), *FText::AsMemory(PackageSize).ToString());
+		}
+		else if (PackageInfo.PackageUpdateType == EConcertPackageUpdateType::Renamed)
+		{
+			UE_LOG(LogConcert, Warning, TEXT("The renamed package '%s' was too large (%s) to be synchronized on other Multi-User clients. It is recommended to leave and persist the session, submit the file to the source control, sync all clients to this submit and create a new session."), *PackageInfo.NewPackageName.ToString(), *FText::AsMemory(PackageSize).ToString());
+		}
+		else // Other events are not expected to trigger a 'too large' error, but log it if something unexpected slips in.
+		{
+			UE_LOG(LogConcert, Error, TEXT("A package event on package '%s' was unexpectedly ignored because the package was too large (%s)."), *PackageInfo.PackageName.ToString(), *FText::AsMemory(PackageSize).ToString());
+		}
+	}
+
 	void RegisterUI()
 	{
 		bHasRegisteredTabSpawners = false;
@@ -836,6 +872,9 @@ private:
 
 		if (MultiUserClient)
 		{
+			MultiUserClient->OnWorkspaceStartup().RemoveAll(this);
+			MultiUserClient->OnWorkspaceShutdown().RemoveAll(this);
+
 			MultiUserClient->GetConcertClient()->OnGetPreConnectionTasks().RemoveAll(this);
 			MultiUserClient->Shutdown();
 			MultiUserClient.Reset();
diff --git a/Engine/Plugins/Developer/Concert/MultiUserClient/Source/MultiUserClient/Private/Widgets/SSessionHistory.cpp b/Engine/Plugins/Developer/Concert/MultiUserClient/Source/MultiUserClient/Private/Widgets/SSessionHistory.cpp
--- a/Engine/Plugins/Developer/Concert/MultiUserClient/Source/MultiUserClient/Private/Widgets/SSessionHistory.cpp
+++ b/Engine/Plugins/Developer/Concert/MultiUserClient/Source/MultiUserClient/Private/Widgets/SSessionHistory.cpp
@@ -52,7 +52,7 @@ void SSessionHistory::Construct(const FArguments& InArgs, TSharedPtr<IConcertSyn
 	ActivityListViewOptions = MakeShared<FConcertSessionActivitiesOptions>();
 
 	SAssignNew(ActivityListView, SConcertSessionActivities)
-		.OnGetPackageEvent([this](const FConcertClientSessionActivity& Activity) { return GetPackageEvent(Activity); })
+		.OnGetPackageEvent([this](const FConcertClientSessionActivity& Activity, FConcertSyncPackageEventMetaData& OutEvent) { return GetPackageEvent(Activity, OutEvent); })
 		.OnGetTransactionEvent([this](const FConcertClientSessionActivity& Activity) { return GetTransactionEvent(Activity); })
 		.OnMapActivityToClient([this](FGuid ClientId){ return EndpointClientInfoMap.Find(ClientId); })
 		.HighlightText(this, &SSessionHistory::HighlightSearchedText)
@@ -220,15 +220,15 @@ void SSessionHistory::RegisterWorkspaceHandler()
 	}
 }
 
-TFuture<TOptional<FConcertSyncPackageEvent>> SSessionHistory::GetPackageEvent(const FConcertClientSessionActivity& Activity) const
+bool SSessionHistory::GetPackageEvent(const FConcertClientSessionActivity& Activity, FConcertSyncPackageEventMetaData& OutPackageEvent) const
 {
 	if (TSharedPtr<IConcertClientWorkspace> WorkspacePtr = Workspace.Pin())
 	{
 		// Don't request the package data, the widget only display the meta-data.
-		return WorkspacePtr->FindOrRequestPackageEvent(Activity.Activity.EventId, /*bMetaDataOnly*/true);
+		return WorkspacePtr->FindPackageEvent(Activity.Activity.EventId, OutPackageEvent);
 	}
 
-	return MakeFulfilledPromise<TOptional<FConcertSyncPackageEvent>>().GetFuture(); // The data is not available.
+	return false; // The data is not available.
 }
 
 TFuture<TOptional<FConcertSyncTransactionEvent>> SSessionHistory::GetTransactionEvent(const FConcertClientSessionActivity& Activity) const
diff --git a/Engine/Plugins/Developer/Concert/MultiUserClient/Source/MultiUserClient/Private/Widgets/SSessionHistory.h b/Engine/Plugins/Developer/Concert/MultiUserClient/Source/MultiUserClient/Private/Widgets/SSessionHistory.h
--- a/Engine/Plugins/Developer/Concert/MultiUserClient/Source/MultiUserClient/Private/Widgets/SSessionHistory.h
+++ b/Engine/Plugins/Developer/Concert/MultiUserClient/Source/MultiUserClient/Private/Widgets/SSessionHistory.h
@@ -62,8 +62,8 @@ private:
 	/** Returns the text to highlight when the search bar has a text set. */
 	FText HighlightSearchedText() const;
 
-	/** Returns the specified package event if available. */
-	TFuture<TOptional<FConcertSyncPackageEvent>> GetPackageEvent(const FConcertClientSessionActivity& Activity) const;
+	/** Returns the specified package event (without the package data itself) if available. */
+	bool GetPackageEvent(const FConcertClientSessionActivity& Activity, FConcertSyncPackageEventMetaData& OutPackageEvent) const;
 
 	/** Returns the specified package event if available. */
 	TFuture<TOptional<FConcertSyncTransactionEvent>> GetTransactionEvent(const FConcertClientSessionActivity& Activity) const;
diff --git a/Engine/Source/Programs/CrashReportClient/Private/RecoveryService.cpp b/Engine/Source/Programs/CrashReportClient/Private/RecoveryService.cpp
--- a/Engine/Source/Programs/CrashReportClient/Private/RecoveryService.cpp
+++ b/Engine/Source/Programs/CrashReportClient/Private/RecoveryService.cpp
@@ -16,6 +16,7 @@
 #include "IConcertSyncServerModule.h"
 #include "ConcertMessageData.h"
 #include "Runtime/Launch/Resources/Version.h"
+#include "ConcertLocalFileSharingService.h"
 
 static const TCHAR RecoveryServiceName[] = TEXT("Disaster Recovery Service");
 
@@ -94,6 +95,7 @@ bool FRecoveryService::Startup()
 
 	// Start disaster recovery server.
 	Server = IConcertSyncServerModule::Get().CreateServer(TEXT("DisasterRecovery"), AutoArchiveSessionFilter);
+	Server->SetFileSharingService(MakeShared<FConcertLocalFileSharingService>(Server->GetConcertServer()->GetRole()));
 	Server->Startup(ServerConfig, EConcertSyncSessionFlags::Default_DisasterRecoverySession);
 
 	UE_LOG(CrashReportClientLog, Display, TEXT("%s Initialized (Name: %s, Version: %d.%d, Role: %s)"), RecoveryServiceName, *Server->GetConcertServer()->GetServerInfo().ServerName, ENGINE_MAJOR_VERSION, ENGINE_MINOR_VERSION, *Server->GetConcertServer()->GetRole());
diff --git a/Engine/Source/Programs/UnrealRecoverySvc/Private/UnrealRecoverySvc.cpp b/Engine/Source/Programs/UnrealRecoverySvc/Private/UnrealRecoverySvc.cpp
--- a/Engine/Source/Programs/UnrealRecoverySvc/Private/UnrealRecoverySvc.cpp
+++ b/Engine/Source/Programs/UnrealRecoverySvc/Private/UnrealRecoverySvc.cpp
@@ -2,6 +2,7 @@
 
 #include "ConcertSettings.h"
 #include "ConcertSyncServerLoop.h"
+#include "ConcertLocalFileSharingService.h"
 
 #include "RequiredProgramMainCPPInclude.h"
 
@@ -23,6 +24,7 @@ INT32_MAIN_INT32_ARGC_TCHAR_ARGV()
 	ServerLoopInitArgs.ServiceRole = TEXT("DisasterRecovery");
 	ServerLoopInitArgs.ServiceFriendlyName = TEXT("Disaster Recovery Service");
 	ServerLoopInitArgs.ServiceAutoArchiveSessionFilter.bIncludeIgnoredActivities = true;
+	ServerLoopInitArgs.FileSharingService = MakeShared<FConcertLocalFileSharingService>(ServerLoopInitArgs.ServiceRole);
 	ServerLoopInitArgs.bShowConsole = false;
 
 	ServerLoopInitArgs.GetServerConfigFunc = [&EditorProcessId]() -> const UConcertServerConfig*
