Fixed missing include files reported by CIS in ConcertLocalFileSharingService.cpp and ConcertServerWorkspace.cpp

#rnx

diff --git a/Engine/Plugins/Developer/Concert/ConcertSync/ConcertSyncCore/Source/ConcertSyncCore/Private/ConcertLocalFileSharingService.cpp b/Engine/Plugins/Developer/Concert/ConcertSync/ConcertSyncCore/Source/ConcertSyncCore/Private/ConcertLocalFileSharingService.cpp
--- a/Engine/Plugins/Developer/Concert/ConcertSync/ConcertSyncCore/Source/ConcertSyncCore/Private/ConcertLocalFileSharingService.cpp
+++ b/Engine/Plugins/Developer/Concert/ConcertSync/ConcertSyncCore/Source/ConcertSyncCore/Private/ConcertLocalFileSharingService.cpp
@@ -6,6 +6,9 @@
 #include "Serialization/JsonReader.h"
 #include "Serialization/JsonWriter.h"
 #include "Serialization/JsonSerializer.h"
+#include "Misc/Paths.h"
+#include "HAL/PlatformProcess.h"
+
 
 FConcertLocalFileSharingService::FConcertLocalFileSharingService(const FString& Role)
 	: SharedRootPathname(FPaths::EngineIntermediateDir() / Role / TEXT("TempFileShare")) // The EngineIntermediateDir is expected to be know of both the client and the server.
diff --git a/Engine/Plugins/Developer/Concert/ConcertSync/ConcertSyncServer/Source/ConcertSyncServer/Private/ConcertServerWorkspace.cpp b/Engine/Plugins/Developer/Concert/ConcertSync/ConcertSyncServer/Source/ConcertSyncServer/Private/ConcertServerWorkspace.cpp
--- a/Engine/Plugins/Developer/Concert/ConcertSync/ConcertSyncServer/Source/ConcertSyncServer/Private/ConcertServerWorkspace.cpp
+++ b/Engine/Plugins/Developer/Concert/ConcertSync/ConcertSyncServer/Source/ConcertSyncServer/Private/ConcertServerWorkspace.cpp
@@ -10,7 +10,7 @@
 #include "ConcertTransactionEvents.h"
 #include "ConcertServerDataStore.h"
 #include "ConcertLogGlobal.h"
-
+#include "Serialization/MemoryReader.h"
 #include "Algo/Transform.h"
 
 FConcertServerWorkspace::FConcertServerWorkspace(const TSharedRef<FConcertSyncServerLiveSession>& InLiveSession, TSharedPtr<IConcertFileSharingService> InFileSharingService)
