- Added null check to remove the warning.

#rnx

diff --git a/Engine/Plugins/Developer/Concert/ConcertSync/ConcertSyncCore/Source/ConcertSyncCore/Private/ConcertSyncSessionDatabase.cpp b/Engine/Plugins/Developer/Concert/ConcertSync/ConcertSyncCore/Source/ConcertSyncCore/Private/ConcertSyncSessionDatabase.cpp
--- a/Engine/Plugins/Developer/Concert/ConcertSync/ConcertSyncCore/Source/ConcertSyncCore/Private/ConcertSyncSessionDatabase.cpp
+++ b/Engine/Plugins/Developer/Concert/ConcertSync/ConcertSyncCore/Source/ConcertSyncCore/Private/ConcertSyncSessionDatabase.cpp
@@ -322,7 +322,7 @@ bool WritePackageBlob(FConcertPackageDataStream& PackageDataStream, EPackageData
 		DstAr << PackageDataStream.DataSize;
 
 		// Copy the package data in the destination.
-		if (PackageDataStream.DataSize > 0)
+		if (PackageDataStream.DataSize > 0 && PackageDataStream.DataAr != nullptr)
 		{
 			ConcertUtil::Copy(DstAr, *PackageDataStream.DataAr, PackageDataStream.DataSize);
 		}
