Added missing override method operator<<(FObjectPtr& Obj) to FConcertSyncObjectWriter and FConcertSyncObjectReader to fix the 'not implemented' asserts firing in the FArchive base class implementation.

diff --git a/Engine/Plugins/Developer/Concert/ConcertSync/ConcertSyncCore/Source/ConcertSyncCore/Private/ConcertSyncArchives.cpp b/Engine/Plugins/Developer/Concert/ConcertSync/ConcertSyncCore/Source/ConcertSyncCore/Private/ConcertSyncArchives.cpp
--- a/Engine/Plugins/Developer/Concert/ConcertSync/ConcertSyncCore/Source/ConcertSyncCore/Private/ConcertSyncArchives.cpp
+++ b/Engine/Plugins/Developer/Concert/ConcertSync/ConcertSyncCore/Source/ConcertSyncCore/Private/ConcertSyncArchives.cpp
@@ -131,6 +131,12 @@ FArchive& FConcertSyncObjectWriter::operator<<(FLazyObjectPtr& LazyObjectPtr)
 	return *this;
 }
 
+FArchive& FConcertSyncObjectWriter::operator<<(FObjectPtr& Obj)
+{
+	UObject* RawObjPtr = Obj.Get();
+	return *this << RawObjPtr;
+}
+
 FArchive& FConcertSyncObjectWriter::operator<<(FSoftObjectPtr& AssetPtr)
 {
 	FSoftObjectPath Obj = AssetPtr.ToSoftObjectPath();
@@ -264,6 +270,14 @@ FArchive& FConcertSyncObjectReader::operator<<(FLazyObjectPtr& LazyObjectPtr)
 	return *this;
 }
 
+FArchive& FConcertSyncObjectReader::operator<<(FObjectPtr& Obj)
+{
+	UObject* RawObj = nullptr;
+	*this << RawObj;
+	Obj = RawObj;
+	return *this;
+}
+
 FArchive& FConcertSyncObjectReader::operator<<(FSoftObjectPtr& AssetPtr)
 {
 	FSoftObjectPath Obj;
diff --git a/Engine/Plugins/Developer/Concert/ConcertSync/ConcertSyncCore/Source/ConcertSyncCore/Public/ConcertSyncArchives.h b/Engine/Plugins/Developer/Concert/ConcertSync/ConcertSyncCore/Source/ConcertSyncCore/Public/ConcertSyncArchives.h
--- a/Engine/Plugins/Developer/Concert/ConcertSync/ConcertSyncCore/Source/ConcertSyncCore/Public/ConcertSyncArchives.h
+++ b/Engine/Plugins/Developer/Concert/ConcertSync/ConcertSyncCore/Source/ConcertSyncCore/Public/ConcertSyncArchives.h
@@ -49,6 +49,7 @@ public:
 	//~ Begin FArchive Interface
 	virtual FArchive& operator<<(UObject*& Obj) override;
 	virtual FArchive& operator<<(FLazyObjectPtr& LazyObjectPtr) override;
+	virtual FArchive& operator<<(FObjectPtr& Obj) override;
 	virtual FArchive& operator<<(FSoftObjectPtr& AssetPtr) override;
 	virtual FArchive& operator<<(FSoftObjectPath& AssetPtr) override;
 	virtual FArchive& operator<<(FWeakObjectPtr& Value) override;
@@ -76,6 +77,7 @@ public:
 	//~ Begin FArchive Interface
 	virtual FArchive& operator<<(UObject*& Obj) override;
 	virtual FArchive& operator<<(FLazyObjectPtr& LazyObjectPtr) override;
+	virtual FArchive& operator<<(FObjectPtr& Obj) override;
 	virtual FArchive& operator<<(FSoftObjectPtr& AssetPtr) override;
 	virtual FArchive& operator<<(FSoftObjectPath& AssetPtr) override;
 	virtual FArchive& operator<<(FWeakObjectPtr& Value) override;
