Use a pointer type to change the archive writer that we are using.

I think the older compiler was doing a copy assignment vs a shallow reference copy.  This seems to work correctly on newer compilers but the binary version is using an older one: 19.24.28315.00

diff --git a/Engine/Plugins/Developer/Concert/ConcertSync/ConcertSyncCore/Source/ConcertSyncCore/Private/ConcertSyncSessionDatabase.cpp b/Engine/Plugins/Developer/Concert/ConcertSync/ConcertSyncCore/Source/ConcertSyncCore/Private/ConcertSyncSessionDatabase.cpp
--- a/Engine/Plugins/Developer/Concert/ConcertSync/ConcertSyncCore/Source/ConcertSyncCore/Private/ConcertSyncSessionDatabase.cpp
+++ b/Engine/Plugins/Developer/Concert/ConcertSync/ConcertSyncCore/Source/ConcertSyncCore/Private/ConcertSyncSessionDatabase.cpp
@@ -397,7 +397,7 @@ void WritePackageUncompressed(FConcertPackageDataStream& PackageDataStream, FArc
 	}
 }
 
-WritePackageResult WritePackage(FConcertPackageDataStream& PackageDataStream, FArchive& DstAr)
+WritePackageResult WritePackage(FConcertPackageDataStream& PackageDataStream, FArchive* DstAr)
 {
 	SCOPED_CONCERT_TRACE(ConcertSyncSessionDatabase_WritePackage);
 
@@ -407,10 +407,10 @@ WritePackageResult WritePackage(FConcertPackageDataStream& PackageDataStream, FA
 	EPackageDataFormat PackageDataFormat = GetPackageDataFormat(PackageDataStream.DataSize);
 	if (WillUseMemoryWriter(PackageDataStream.DataSize))
 	{
-		DstAr = MemWriter;
+		DstAr = &MemWriter;
 	}
 
-	check(DstAr.IsSaving() && DstAr.Tell() == 0);
+	check(DstAr && DstAr->IsSaving() && DstAr->Tell() == 0);
 
 	// The package data stream may be null to write an 'empty' package.
 	check(PackageDataStream.DataAr == nullptr || PackageDataStream.DataAr->IsLoading());
@@ -418,18 +418,18 @@ WritePackageResult WritePackage(FConcertPackageDataStream& PackageDataStream, FA
 	 // If the package data archive is null, the size must be zero.
 	check(PackageDataStream.DataAr != nullptr || (PackageDataStream.DataAr == nullptr && PackageDataStream.DataSize == 0 && PackageDataStream.DataBlob == nullptr));
 
-	WritePackageHeader(DstAr, PackageDataFormat);
+	WritePackageHeader(*DstAr, PackageDataFormat);
 	if (PackageDataFormat == EPackageDataFormat::PackageDataCompressed)
 	{
-		WritePackageCompressed(PackageDataStream, DstAr);
+		WritePackageCompressed(PackageDataStream, *DstAr);
 		OutResult.bShouldCache = ShouldCachePackageBlob(OutResult.CacheData.Num());
 	}
 	else
 	{
-		WritePackageUncompressed(PackageDataStream, DstAr);
+		WritePackageUncompressed(PackageDataStream, *DstAr);
 		OutResult.bShouldCache = ShouldCachePackageBlob(PackageDataStream.DataSize);
 	}
-	WritePackageFooter(DstAr);
+	WritePackageFooter(*DstAr);
 
 	return OutResult;
 }
@@ -2735,7 +2735,7 @@ bool FConcertSyncSessionDatabase::SavePackage(const FString& InDstPackageBlobPat
 		if (PackageDataUtil::WillUseMemoryWriter(InPackageDataStream.DataSize))
 		{
 			FArchive NullArchive;
-			PackageDataUtil::WritePackageResult Result = PackageDataUtil::WritePackage(InPackageDataStream, NullArchive);
+			PackageDataUtil::WritePackageResult Result = PackageDataUtil::WritePackage(InPackageDataStream, &NullArchive);
 			return HandleWritePackageResult(InDstPackageBlobPathname, PackageFileCache, MoveTemp(Result));
 		}
 		else
@@ -2743,7 +2743,7 @@ bool FConcertSyncSessionDatabase::SavePackage(const FString& InDstPackageBlobPat
 			TUniquePtr<FArchive> DstAr(IFileManager::Get().CreateFileWriter(*InDstPackageBlobPathname));
 			if (DstAr)
 			{
-				PackageDataUtil::WritePackage(InPackageDataStream, *DstAr);
+				PackageDataUtil::WritePackage(InPackageDataStream, DstAr.Get());
 				return !DstAr->IsError();
 			}
 		}
@@ -2828,7 +2828,7 @@ void FConcertSyncSessionDatabase::ScheduleAsyncWrite(const FString& InDstPackage
 		FConcertPackageDataStream& InPackageDataStream = SharedStream->PackageStream;
 		FMemoryReader PackageBlobAr(SharedStream->PackageData);
 		InPackageDataStream.DataAr = &PackageBlobAr;
-		SharedStream->Result = PackageDataUtil::WritePackage(InPackageDataStream, *DstAr);
+		SharedStream->Result = PackageDataUtil::WritePackage(InPackageDataStream, DstAr.Get());
 		return !DstAr->IsError();
 	});
 }
