Make sure we call result handler when flushing the async task list.

diff --git a/Engine/Plugins/Developer/Concert/ConcertSync/ConcertSyncCore/Source/ConcertSyncCore/Private/ConcertSyncSessionDatabase.cpp b/Engine/Plugins/Developer/Concert/ConcertSync/ConcertSyncCore/Source/ConcertSyncCore/Private/ConcertSyncSessionDatabase.cpp
--- a/Engine/Plugins/Developer/Concert/ConcertSync/ConcertSyncCore/Source/ConcertSyncCore/Private/ConcertSyncSessionDatabase.cpp
+++ b/Engine/Plugins/Developer/Concert/ConcertSync/ConcertSyncCore/Source/ConcertSyncCore/Private/ConcertSyncSessionDatabase.cpp
@@ -2694,7 +2694,7 @@ bool FConcertSyncSessionDatabase::LoadTransaction(const FString& InTransactionFi
 	return false;
 }
 
-bool HandleWritePackageResult(const FString& InDstPackageBlobPathname, TUniquePtr<FConcertFileCache>& Cache, PackageDataUtil::WritePackageResult&& Result)
+bool HandleWritePackageResult(const FString& InDstPackageBlobPathname, TUniquePtr<FConcertFileCache>& Cache, PackageDataUtil::WritePackageResult Result)
 {
 	if (Result.CacheData.Num()>0)
 	{
@@ -2708,6 +2708,7 @@ bool HandleWritePackageResult(const FString& InDstPackageBlobPathname, TUniquePt
 			DstAr->Serialize(Result.CacheData.GetData(), Result.CacheData.Num());
 			return !DstAr->IsError();
 		}
+		UE_LOG(LogConcert, Warning, TEXT("Failed to handle package write result of file '%s'."), *InDstPackageBlobPathname);
 		return false;
 	}
 	return true;
@@ -2787,7 +2788,11 @@ void FConcertSyncSessionDatabase::FlushAsynchronousTasks()
 {
 	for (const auto& Item : DeferredLargePackageIO)
 	{
-		if(!Item.Value->AsyncTask.Get())
+		if (Item.Value->AsyncTask.Get())
+		{
+			HandleWritePackageResult(Item.Key, PackageFileCache, MoveTemp(Item.Value->Result));
+		}
+		else
 		{
 			UE_LOG(LogConcert, Error, TEXT("Async task failed to write package %s"), *Item.Value->CachedPackageName);
 		}
@@ -2801,10 +2806,7 @@ void FConcertSyncSessionDatabase::UpdateAsynchronousTasks()
 	{
 		if (It->Value->AsyncTask.IsReady())
 		{
-			if (HandleWritePackageResult(It->Key, PackageFileCache,MoveTemp(It->Value->Result)) == false)
-			{
-				UE_LOG(LogConcert, Warning, TEXT("Failed to handle package write result of file '%s'."), *It->Key);
-			}
+			HandleWritePackageResult(It->Key, PackageFileCache,MoveTemp(It->Value->Result));
 			It.RemoveCurrent();
 		}
 	}
