Make sure we explicitly mark the shared pointer as thread safe and move the map into the implementation hidden via a pointer to implementation.

#[fyi] matt.johnson

diff --git a/Engine/Plugins/Developer/Concert/ConcertSync/ConcertSyncCore/Source/ConcertSyncCore/Private/ConcertSyncSessionDatabase.cpp b/Engine/Plugins/Developer/Concert/ConcertSync/ConcertSyncCore/Source/ConcertSyncCore/Private/ConcertSyncSessionDatabase.cpp
--- a/Engine/Plugins/Developer/Concert/ConcertSync/ConcertSyncCore/Source/ConcertSyncCore/Private/ConcertSyncSessionDatabase.cpp
+++ b/Engine/Plugins/Developer/Concert/ConcertSync/ConcertSyncCore/Source/ConcertSyncCore/Private/ConcertSyncSessionDatabase.cpp
@@ -10,6 +10,7 @@
 #include "SQLiteDatabase.h"
 #include "HAL/FileManager.h"
 #include "Templates/SharedPointer.h"
+#include "Templates/SharedPointerInternals.h"
 #include "Templates/UniquePtr.h"
 #include "UObject/ObjectMacros.h"
 #include "UObject/StructOnScope.h"
@@ -46,6 +47,29 @@ struct FConcertPackageAsyncDataStream
 	TArray<uint8>			    PackageData;
 };
 
+using FConcertPackageAsyncDataStreamPtr = TSharedPtr<FConcertPackageAsyncDataStream, ESPMode::ThreadSafe>;
+
+struct FDeferredLargePackageIOImpl
+{
+	using FAsyncDataStreamMap = TMap<FString,FConcertPackageAsyncDataStreamPtr>;
+	FAsyncDataStreamMap& GetMap()
+	{
+		return DeferredLargePackageIO;
+	}
+
+	const FAsyncDataStreamMap& GetMap() const
+	{
+		return DeferredLargePackageIO;
+	}
+
+	void Clear()
+	{
+		DeferredLargePackageIO = {};
+	}
+
+	FAsyncDataStreamMap DeferredLargePackageIO;
+};
+
 namespace TransactionDataUtil
 {
 
@@ -1508,7 +1532,8 @@ private:
 
 /** Defined here where TUniquePtr can see the definition of FConcertFileCache and FConcertSyncSessionDatabaseStatements as the TUniquePtr constructor/destructor cannot work with a forward declared type */
 FConcertSyncSessionDatabase::FConcertSyncSessionDatabase()
-	: Database(MakeUnique<FSQLiteDatabase>())
+	: Database(MakeUnique<FSQLiteDatabase>()),
+	  DeferredLargePackageIOPtr(MakeUnique<FDeferredLargePackageIOImpl>())
 {
 }
 
@@ -2756,10 +2781,9 @@ bool FConcertSyncSessionDatabase::SavePackage(const FString& InDstPackageBlobPat
 bool FConcertSyncSessionDatabase::LoadPackage(const FString& InPackageBlobFilename, const TFunctionRef<void(FConcertPackageDataStream&)>& PackageDataStreamFn) const
 {
 	// First check the deferred file IO buffer before looking on disk.
-	using FAsyncPtr = TSharedPtr<FConcertPackageAsyncDataStream>;
-	if(FAsyncPtr const* Found = DeferredLargePackageIO.Find(InPackageBlobFilename))
+	if(FConcertPackageAsyncDataStreamPtr const* Found = DeferredLargePackageIOPtr->GetMap().Find(InPackageBlobFilename))
 	{
-		FAsyncPtr Owned = *Found;
+		FConcertPackageAsyncDataStreamPtr Owned = *Found;
 		FConcertPackageDataStream Stream = Owned->PackageStream;
 		FMemoryReader PackageBlobAr(Owned->PackageData);
 		Stream.DataAr = &PackageBlobAr;
@@ -2786,7 +2810,7 @@ bool FConcertSyncSessionDatabase::LoadPackage(const FString& InPackageBlobFilena
 
 void FConcertSyncSessionDatabase::FlushAsynchronousTasks()
 {
-	for (const auto& Item : DeferredLargePackageIO)
+	for (const auto& Item : DeferredLargePackageIOPtr->GetMap())
 	{
 		if (Item.Value->AsyncTask.Get())
 		{
@@ -2797,12 +2821,12 @@ void FConcertSyncSessionDatabase::FlushAsynchronousTasks()
 			UE_LOG(LogConcert, Error, TEXT("Async task failed to write package %s"), *Item.Value->CachedPackageName);
 		}
 	}
-	DeferredLargePackageIO = {};
+	DeferredLargePackageIOPtr->Clear();
 }
 
 void FConcertSyncSessionDatabase::UpdateAsynchronousTasks()
 {
-	for (auto It = DeferredLargePackageIO.CreateIterator(); It; ++It)
+	for (auto It = DeferredLargePackageIOPtr->GetMap().CreateIterator(); It; ++It)
 	{
 		if (It->Value->AsyncTask.IsReady())
 		{
@@ -2816,14 +2840,13 @@ void FConcertSyncSessionDatabase::ScheduleAsyncWrite(const FString& InDstPackage
 {
 	check(InPackageDataStream.DataBlob && InPackageDataStream.DataBlob->Num() > 0);
 
-	using FAsyncPtr = TSharedPtr<FConcertPackageAsyncDataStream>;
-	FAsyncPtr SharedStream = MakeShared<FConcertPackageAsyncDataStream>(InPackageDataStream, InDstPackageBlobPathname);
-	if(FAsyncPtr* Item = DeferredLargePackageIO.Find(InDstPackageBlobPathname))
+	FConcertPackageAsyncDataStreamPtr SharedStream = MakeShared<FConcertPackageAsyncDataStream,ESPMode::ThreadSafe>(InPackageDataStream, InDstPackageBlobPathname);
+	if(FConcertPackageAsyncDataStreamPtr* Item = DeferredLargePackageIOPtr->GetMap().Find(InDstPackageBlobPathname))
 	{
 		(*Item)->AsyncTask.Get();
 		HandleWritePackageResult(InDstPackageBlobPathname, PackageFileCache, MoveTemp((*Item)->Result));
 	}
-	DeferredLargePackageIO.FindOrAdd(InDstPackageBlobPathname) = SharedStream;
+	DeferredLargePackageIOPtr->GetMap().FindOrAdd(InDstPackageBlobPathname) = SharedStream;
 	SharedStream->AsyncTask = Async(EAsyncExecution::TaskGraph, [InDstPackageBlobPathname,SharedStream]()
 	{
 		TUniquePtr<FArchive> DstAr(IFileManager::Get().CreateFileWriter(*InDstPackageBlobPathname));
diff --git a/Engine/Plugins/Developer/Concert/ConcertSync/ConcertSyncCore/Source/ConcertSyncCore/Public/ConcertSyncSessionDatabase.h b/Engine/Plugins/Developer/Concert/ConcertSync/ConcertSyncCore/Source/ConcertSyncCore/Public/ConcertSyncSessionDatabase.h
--- a/Engine/Plugins/Developer/Concert/ConcertSync/ConcertSyncCore/Source/ConcertSyncCore/Public/ConcertSyncSessionDatabase.h
+++ b/Engine/Plugins/Developer/Concert/ConcertSync/ConcertSyncCore/Source/ConcertSyncCore/Public/ConcertSyncSessionDatabase.h
@@ -4,6 +4,7 @@
 
 #include "CoreMinimal.h"
 #include "ConcertSyncSessionTypes.h"
+#include "Templates/SharedPointerInternals.h"
 
 class FConcertFileCache;
 class FConcertSyncSessionDatabaseStatements;
@@ -933,7 +934,7 @@ private:
 	/** Internal SQLite database */
 	TUniquePtr<FSQLiteDatabase> Database;
 
-	TMap<FString,TSharedPtr<struct FConcertPackageAsyncDataStream>> DeferredLargePackageIO;
+	TUniquePtr<struct FDeferredLargePackageIOImpl> DeferredLargePackageIOPtr;
 };
 
 namespace ConcertSyncSessionDatabaseFilterUtil
