Update concert sync archive to allow delegate to handle remapping of object paths

diff --git a/Engine/Plugins/Developer/Concert/ConcertSync/ConcertSyncCore/Source/ConcertSyncCore/Private/ConcertSyncArchives.cpp b/Engine/Plugins/Developer/Concert/ConcertSync/ConcertSyncCore/Source/ConcertSyncCore/Private/ConcertSyncArchives.cpp
--- a/Engine/Plugins/Developer/Concert/ConcertSync/ConcertSyncCore/Source/ConcertSyncCore/Private/ConcertSyncArchives.cpp
+++ b/Engine/Plugins/Developer/Concert/ConcertSync/ConcertSyncCore/Source/ConcertSyncCore/Private/ConcertSyncArchives.cpp
@@ -48,6 +48,13 @@ namespace ConcertSyncUtil
 
 FString FConcertSyncWorldRemapper::RemapObjectPathName(const FString& InObjectPathName) const
 {
+	if (RemapDelegate.IsBound())
+	{
+		FString Result;
+		RemapDelegate.Execute(Result);
+		return MoveTemp(Result);
+	}
+
 	return HasMapping() ? InObjectPathName.Replace(*SourceWorldPathName, *DestWorldPathName) : InObjectPathName;
 }
 
@@ -58,7 +65,7 @@ bool FConcertSyncWorldRemapper::ObjectBelongsToWorld(const FString& InObjectPath
 
 bool FConcertSyncWorldRemapper::HasMapping() const
 {
-	return SourceWorldPathName.Len() > 0 && DestWorldPathName.Len() > 0;
+	return RemapDelegate.IsBound() || (SourceWorldPathName.Len() > 0 && DestWorldPathName.Len() > 0 || RemapDelegate.IsBound());
 }
 
 FConcertSyncObjectWriter::FConcertSyncObjectWriter(FConcertLocalIdentifierTable* InLocalIdentifierTable, UObject* InObj, TArray<uint8>& OutBytes, const bool InIncludeEditorOnlyData, const bool InSkipAssets)
diff --git a/Engine/Plugins/Developer/Concert/ConcertSync/ConcertSyncCore/Source/ConcertSyncCore/Public/ConcertSyncArchives.h b/Engine/Plugins/Developer/Concert/ConcertSync/ConcertSyncCore/Source/ConcertSyncCore/Public/ConcertSyncArchives.h
--- a/Engine/Plugins/Developer/Concert/ConcertSync/ConcertSyncCore/Source/ConcertSyncCore/Public/ConcertSyncArchives.h
+++ b/Engine/Plugins/Developer/Concert/ConcertSync/ConcertSyncCore/Source/ConcertSyncCore/Public/ConcertSyncArchives.h
@@ -7,6 +7,8 @@
 
 struct FConcertSessionVersionInfo;
 
+DECLARE_DELEGATE_OneParam(FConcertSyncRemapObjectPath, FString& /*ObjectPath*/)
+
 namespace ConcertSyncUtil
 {
 	 bool CONCERTSYNCCORE_API ShouldSkipTransientProperty(const FProperty* Property);
@@ -30,6 +32,8 @@ public:
 
 	bool HasMapping() const;
 
+	FConcertSyncRemapObjectPath RemapDelegate;
+
 private:
 	FString SourceWorldPathName;
 	FString DestWorldPathName;
