Fix compile warning
#rnx

diff --git a/Engine/Plugins/Developer/Concert/ConcertSync/ConcertSyncCore/Source/ConcertSyncCore/Private/ConcertSyncArchives.cpp b/Engine/Plugins/Developer/Concert/ConcertSync/ConcertSyncCore/Source/ConcertSyncCore/Private/ConcertSyncArchives.cpp
--- a/Engine/Plugins/Developer/Concert/ConcertSync/ConcertSyncCore/Source/ConcertSyncCore/Private/ConcertSyncArchives.cpp
+++ b/Engine/Plugins/Developer/Concert/ConcertSync/ConcertSyncCore/Source/ConcertSyncCore/Private/ConcertSyncArchives.cpp
@@ -65,7 +65,7 @@ bool FConcertSyncWorldRemapper::ObjectBelongsToWorld(const FString& InObjectPath
 
 bool FConcertSyncWorldRemapper::HasMapping() const
 {
-	return RemapDelegate.IsBound() || (SourceWorldPathName.Len() > 0 && DestWorldPathName.Len() > 0 || RemapDelegate.IsBound());
+	return RemapDelegate.IsBound() || (SourceWorldPathName.Len() > 0 && DestWorldPathName.Len() > 0);
 }
 
 FConcertSyncObjectWriter::FConcertSyncObjectWriter(FConcertLocalIdentifierTable* InLocalIdentifierTable, UObject* InObj, TArray<uint8>& OutBytes, const bool InIncludeEditorOnlyData, const bool InSkipAssets)
