Concert sync object writer can now be given delegate to remap object paths

diff --git a/Engine/Plugins/Developer/Concert/ConcertSync/ConcertSyncCore/Source/ConcertSyncCore/Private/ConcertSyncArchives.cpp b/Engine/Plugins/Developer/Concert/ConcertSync/ConcertSyncCore/Source/ConcertSyncCore/Private/ConcertSyncArchives.cpp
--- a/Engine/Plugins/Developer/Concert/ConcertSync/ConcertSyncCore/Source/ConcertSyncCore/Private/ConcertSyncArchives.cpp
+++ b/Engine/Plugins/Developer/Concert/ConcertSync/ConcertSyncCore/Source/ConcertSyncCore/Private/ConcertSyncArchives.cpp
@@ -68,10 +68,11 @@ bool FConcertSyncWorldRemapper::HasMapping() const
 	return RemapDelegate.IsBound() || (SourceWorldPathName.Len() > 0 && DestWorldPathName.Len() > 0);
 }
 
-FConcertSyncObjectWriter::FConcertSyncObjectWriter(FConcertLocalIdentifierTable* InLocalIdentifierTable, UObject* InObj, TArray<uint8>& OutBytes, const bool InIncludeEditorOnlyData, const bool InSkipAssets)
+FConcertSyncObjectWriter::FConcertSyncObjectWriter(FConcertLocalIdentifierTable* InLocalIdentifierTable, UObject* InObj, TArray<uint8>& OutBytes, const bool InIncludeEditorOnlyData, const bool InSkipAssets, const FConcertSyncRemapObjectPath& InRemapDelegate)
 	: FConcertIdentifierWriter(InLocalIdentifierTable, OutBytes, /*bIsPersistent*/false)
 	, bSkipAssets(InSkipAssets)
 	, ShouldSkipPropertyFunc()
+	, RemapObjectPathDelegate(InRemapDelegate)
 {
 	ArIgnoreClassRef = false;
 	ArIgnoreArchetypeRef = false;
@@ -89,6 +90,11 @@ FConcertSyncObjectWriter::FConcertSyncObjectWriter(FConcertLocalIdentifierTable*
 #endif // USE_STABLE_LOCALIZATION_KEYS
 }
 
+FConcertSyncObjectWriter::FConcertSyncObjectWriter(FConcertLocalIdentifierTable* InLocalIdentifierTable, UObject* InObj, TArray<uint8>& OutBytes, const bool InIncludeEditorOnlyData, const bool InSkipAssets)
+	: FConcertSyncObjectWriter(InLocalIdentifierTable, InObj, OutBytes, InIncludeEditorOnlyData, InSkipAssets, FConcertSyncRemapObjectPath())
+{
+}
+
 void FConcertSyncObjectWriter::SerializeObject(UObject* InObject, const TArray<FName>* InPropertyNamesToWrite)
 {
 	if (InPropertyNamesToWrite)
@@ -121,7 +127,16 @@ FArchive& FConcertSyncObjectWriter::operator<<(UObject*& Obj)
 	FName ObjPath;
 	if (Obj)
 	{
-		ObjPath = (Obj->IsAsset() && bSkipAssets) ? SkipAssetsMarker : *Obj->GetPathName();
+		if (bSkipAssets && Obj->IsAsset())
+		{
+			ObjPath = SkipAssetsMarker;
+		}
+		else
+		{
+			FString ObjectPathString = Obj->GetPathName();
+			RemapObjectPathDelegate.ExecuteIfBound(ObjectPathString);
+			ObjPath = FName(ObjectPathString);
+		}
 	}
 
 	*this << ObjPath;
@@ -153,7 +168,17 @@ FArchive& FConcertSyncObjectWriter::operator<<(FSoftObjectPtr& AssetPtr)
 
 FArchive& FConcertSyncObjectWriter::operator<<(FSoftObjectPath& AssetPtr)
 {
-	FName ObjPath = bSkipAssets ? SkipAssetsMarker : *AssetPtr.ToString();
+	FName ObjPath;
+	if (bSkipAssets)
+	{
+		ObjPath = SkipAssetsMarker;
+	}
+	else
+	{
+		FString ObjectPathString = AssetPtr.ToString();
+		RemapObjectPathDelegate.ExecuteIfBound(ObjectPathString);
+		ObjPath = FName(ObjectPathString);
+	}
 	*this << ObjPath;
 	return *this;
 }
diff --git a/Engine/Plugins/Developer/Concert/ConcertSync/ConcertSyncCore/Source/ConcertSyncCore/Public/ConcertSyncArchives.h b/Engine/Plugins/Developer/Concert/ConcertSync/ConcertSyncCore/Source/ConcertSyncCore/Public/ConcertSyncArchives.h
--- a/Engine/Plugins/Developer/Concert/ConcertSync/ConcertSyncCore/Source/ConcertSyncCore/Public/ConcertSyncArchives.h
+++ b/Engine/Plugins/Developer/Concert/ConcertSync/ConcertSyncCore/Source/ConcertSyncCore/Public/ConcertSyncArchives.h
@@ -43,6 +43,7 @@ private:
 class CONCERTSYNCCORE_API FConcertSyncObjectWriter : public FConcertIdentifierWriter
 {
 public:
+	FConcertSyncObjectWriter(FConcertLocalIdentifierTable* InLocalIdentifierTable, UObject* InObj, TArray<uint8>& OutBytes, const bool InIncludeEditorOnlyData, const bool InSkipAssets, const FConcertSyncRemapObjectPath& InRemapDelegate);
 	FConcertSyncObjectWriter(FConcertLocalIdentifierTable* InLocalIdentifierTable, UObject* InObj, TArray<uint8>& OutBytes, const bool InIncludeEditorOnlyData, const bool InSkipAssets);
 
 	void SerializeObject(UObject* InObject, const TArray<FName>* InPropertyNamesToWrite = nullptr);
@@ -65,6 +66,7 @@ private:
 	typedef TFunction<bool(const FProperty*)> FShouldSkipPropertyFunc;
 	bool bSkipAssets;
 	FShouldSkipPropertyFunc ShouldSkipPropertyFunc;
+	FConcertSyncRemapObjectPath RemapObjectPathDelegate;
 };
 
 /** Archive for reading objects that can been received from another instance via Concert */
