Fix an issue in the concert sync world remapper where paths for object pointers wouldn't resolve properly and return an empty string if a remap delegate was set.

diff --git a/Engine/Plugins/Developer/Concert/ConcertSync/ConcertSyncCore/Source/ConcertSyncCore/Private/ConcertSyncArchives.cpp b/Engine/Plugins/Developer/Concert/ConcertSync/ConcertSyncCore/Source/ConcertSyncCore/Private/ConcertSyncArchives.cpp
--- a/Engine/Plugins/Developer/Concert/ConcertSync/ConcertSyncCore/Source/ConcertSyncCore/Private/ConcertSyncArchives.cpp
+++ b/Engine/Plugins/Developer/Concert/ConcertSync/ConcertSyncCore/Source/ConcertSyncCore/Private/ConcertSyncArchives.cpp
@@ -50,7 +50,7 @@ FString FConcertSyncWorldRemapper::RemapObjectPathName(const FString& InObjectPa
 {
 	if (RemapDelegate.IsBound())
 	{
-		FString Result;
+		FString Result = InObjectPathName;
 		RemapDelegate.Execute(Result);
 		return MoveTemp(Result);
 	}
