Fix issue in concert transaction reader where object ptr properties would get set to null if skip assets is true

diff --git a/Engine/Plugins/Developer/Concert/ConcertSync/ConcertSyncCore/Source/ConcertSyncCore/Private/ConcertSyncArchives.cpp b/Engine/Plugins/Developer/Concert/ConcertSync/ConcertSyncCore/Source/ConcertSyncCore/Private/ConcertSyncArchives.cpp
--- a/Engine/Plugins/Developer/Concert/ConcertSync/ConcertSyncCore/Source/ConcertSyncCore/Private/ConcertSyncArchives.cpp
+++ b/Engine/Plugins/Developer/Concert/ConcertSync/ConcertSyncCore/Source/ConcertSyncCore/Private/ConcertSyncArchives.cpp
@@ -308,7 +308,7 @@ FArchive& FConcertSyncObjectReader::operator<<(FLazyObjectPtr& LazyObjectPtr)
 
 FArchive& FConcertSyncObjectReader::operator<<(FObjectPtr& Obj)
 {
-	UObject* RawObj = nullptr;
+	UObject* RawObj = Obj.Get();
 	*this << RawObj;
 	Obj = RawObj;
 	return *this;
