Fix missing errors with forward declared types

diff --git a/Engine/Plugins/Developer/Concert/ConcertSync/ConcertSyncCore/Source/ConcertSyncCore/Private/ConcertSyncSessionDatabase.cpp b/Engine/Plugins/Developer/Concert/ConcertSync/ConcertSyncCore/Source/ConcertSyncCore/Private/ConcertSyncSessionDatabase.cpp
--- a/Engine/Plugins/Developer/Concert/ConcertSync/ConcertSyncCore/Source/ConcertSyncCore/Private/ConcertSyncSessionDatabase.cpp
+++ b/Engine/Plugins/Developer/Concert/ConcertSync/ConcertSyncCore/Source/ConcertSyncCore/Private/ConcertSyncSessionDatabase.cpp
@@ -1533,12 +1533,15 @@ private:
 /** Defined here where TUniquePtr can see the definition of FConcertFileCache and FConcertSyncSessionDatabaseStatements as the TUniquePtr constructor/destructor cannot work with a forward declared type */
 FConcertSyncSessionDatabase::FConcertSyncSessionDatabase()
 	: Database(MakeUnique<FSQLiteDatabase>()),
-	  DeferredLargePackageIOPtr(MakeUnique<FDeferredLargePackageIOImpl>())
+	  DeferredLargePackageIOPtr(MakePimpl<FDeferredLargePackageIOImpl>())
 {
 }
 
 FConcertSyncSessionDatabase::~FConcertSyncSessionDatabase() = default;
 
+FConcertSyncSessionDatabase::FConcertSyncSessionDatabase(FConcertSyncSessionDatabase&&) = default;
+FConcertSyncSessionDatabase& FConcertSyncSessionDatabase::operator=(FConcertSyncSessionDatabase&&) = default;
+
 bool FConcertSyncSessionDatabase::IsValid() const
 {
 	return Database->IsValid();
diff --git a/Engine/Plugins/Developer/Concert/ConcertSync/ConcertSyncCore/Source/ConcertSyncCore/Public/ConcertSyncSessionDatabase.h b/Engine/Plugins/Developer/Concert/ConcertSync/ConcertSyncCore/Source/ConcertSyncCore/Public/ConcertSyncSessionDatabase.h
--- a/Engine/Plugins/Developer/Concert/ConcertSync/ConcertSyncCore/Source/ConcertSyncCore/Public/ConcertSyncSessionDatabase.h
+++ b/Engine/Plugins/Developer/Concert/ConcertSync/ConcertSyncCore/Source/ConcertSyncCore/Public/ConcertSyncSessionDatabase.h
@@ -38,8 +38,8 @@ public:
 	FConcertSyncSessionDatabase(const FConcertSyncSessionDatabase&) = delete;
 	FConcertSyncSessionDatabase& operator=(const FConcertSyncSessionDatabase&) = delete;
 
-	FConcertSyncSessionDatabase(FConcertSyncSessionDatabase&&) = default;
-	FConcertSyncSessionDatabase& operator=(FConcertSyncSessionDatabase&&) = default;
+	FConcertSyncSessionDatabase(FConcertSyncSessionDatabase&&);
+	FConcertSyncSessionDatabase& operator=(FConcertSyncSessionDatabase&&);
 
 	/**
 	 * Is this a valid database? (ie, has been successfully opened).
@@ -952,7 +952,7 @@ private:
 	/** Internal SQLite database */
 	TUniquePtr<FSQLiteDatabase> Database;
 
-	TUniquePtr<struct FDeferredLargePackageIOImpl> DeferredLargePackageIOPtr;
+	TPimplPtr<struct FDeferredLargePackageIOImpl> DeferredLargePackageIOPtr;
 };
 
 namespace ConcertSyncSessionDatabaseFilterUtil
