Protect the level from getting marked dirty when it is not suppose to based on using external actors in the file.

diff --git a/Engine/Plugins/Developer/Concert/ConcertSync/ConcertSyncClient/Source/ConcertSyncClient/Private/ConcertClientTransactionBridge.cpp b/Engine/Plugins/Developer/Concert/ConcertSync/ConcertSyncClient/Source/ConcertSyncClient/Private/ConcertClientTransactionBridge.cpp
--- a/Engine/Plugins/Developer/Concert/ConcertSync/ConcertSyncClient/Source/ConcertSyncClient/Private/ConcertClientTransactionBridge.cpp
+++ b/Engine/Plugins/Developer/Concert/ConcertSync/ConcertSyncClient/Source/ConcertSyncClient/Private/ConcertClientTransactionBridge.cpp
@@ -270,8 +270,19 @@ void ProcessTransactionEvent(const FConcertTransactionEventBase& InEvent, const
 		// Notify before changing anything
 		if (!bIsSnapshot || TransactionAnnotation)
 		{
+			ULevel* Level = Cast<ULevel>(TransactionObject->GetOuter());
+			const bool bLevelIsDirty = Level ? Level->GetPackage()->IsDirty() : false;
+
 			// Transaction annotations require us to invoke the redo flow (even for snapshots!) as that's the only thing that can apply the annotation
 			TransactionObject->PreEditUndo();
+
+			// Levels are immune from dirty changes when using external objects.  See ULevel::PreEditUndo() If we
+			// modified any dirty flags as a result of the PreEditUndo then restore it back here as if it didn't happen.
+			//
+			if (Level && Level->IsUsingExternalObjects() && TransactionObject->IsPackageExternal())
+			{
+				Level->GetPackage()->SetDirtyFlag(bLevelIsDirty);
+			}
 		}
 
 		// We need to manually call OnPreObjectPropertyChanged as PreEditUndo calls the PreEditChange version that skips it, but we have things that rely on it being called
diff --git a/Engine/Plugins/Developer/Concert/ConcertSync/ConcertSyncCore/Config/BaseConcertSyncCore.ini b/Engine/Plugins/Developer/Concert/ConcertSync/ConcertSyncCore/Config/BaseConcertSyncCore.ini
--- a/Engine/Plugins/Developer/Concert/ConcertSync/ConcertSyncCore/Config/BaseConcertSyncCore.ini
+++ b/Engine/Plugins/Developer/Concert/ConcertSync/ConcertSyncCore/Config/BaseConcertSyncCore.ini
@@ -1,6 +1,6 @@
 [/Script/ConcertSyncCore.ConcertSyncConfig]
 +IncludeObjectClassFilters=(ObjectOuterClass=None,ObjectClasses=(Class'/Script/Engine.World',Class'/Script/Engine.BookmarkBase',Class'/Script/Engine.MaterialParameterCollection',Class'/Script/MovieScene.MovieSceneSignedObject',Class'/Script/MovieScene.MovieSceneFolder'))
-+IncludeObjectClassFilters=(ObjectOuterClass=Class'/Script/Engine.World',ObjectClasses=(Class'/Script/Engine.Layer',Class'/Script/Engine.Level',Class'/Script/Engine.LevelStreaming',Class'/Script/Engine.Actor',Class'/Script/Engine.ActorComponent',Class'/Script/Engine.Polys',Class'/Script/Engine.Model',Class'/Script/Engine.BrushBuilder',Class'/Script/Foliage.FoliageType'))
++IncludeObjectClassFilters=(ObjectOuterClass=Class'/Script/Engine.World',ObjectClasses=(Class'/Script/Engine.Layer',Class'/Script/Engine.Level',Class'/Script/Engine.LevelStreaming',Class'/Script/Engine.Actor',Class'/Script/Engine.ActorComponent',Class'/Script/Engine.Polys',Class'/Script/Engine.BrushBuilder',Class'/Script/Foliage.FoliageType'))
 +IncludeObjectClassFilters=(ObjectOuterClass=Class'/Script/MovieScene.MovieSceneSignedObject',ObjectClasses=(Class'/Script/Engine.Actor',Class'/Script/Engine.ActorComponent'))
 +ExcludeTransactionClassFilters=(ObjectOuterClass=None,ObjectClasses=(Class'/Script/Landscape.LandscapeProxy',Class'/Script/Landscape.LandscapeHeightfieldCollisionComponent',Class'/Script/Landscape.LandscapeComponent'))
 +AllowedTransientProperties="/Script/Engine.Actor:bHiddenEdTemporary"
