Show package revision and size in package viewer

diff --git a/Engine/Plugins/Developer/Concert/ConcertApp/MultiUserServer/Source/MultiUserServer/Private/Widgets/SessionTabs/Live/PackageViewer/ConcertSessionPackageViewerController.cpp b/Engine/Plugins/Developer/Concert/ConcertApp/MultiUserServer/Source/MultiUserServer/Private/Widgets/SessionTabs/Live/PackageViewer/ConcertSessionPackageViewerController.cpp
--- a/Engine/Plugins/Developer/Concert/ConcertApp/MultiUserServer/Source/MultiUserServer/Private/Widgets/SessionTabs/Live/PackageViewer/ConcertSessionPackageViewerController.cpp
+++ b/Engine/Plugins/Developer/Concert/ConcertApp/MultiUserServer/Source/MultiUserServer/Private/Widgets/SessionTabs/Live/PackageViewer/ConcertSessionPackageViewerController.cpp
@@ -72,7 +72,8 @@ TSharedRef<SConcertSessionPackageViewer> FConcertSessionPackageViewerController:
 {
 	return SNew(SConcertSessionPackageViewer)
 		.GetClientInfo_Raw(this, &FConcertSessionPackageViewerController::GetClientInfo)
-		.GetPackageEvent_Raw(this, &FConcertSessionPackageViewerController::GetPackageEvent);
+		.GetPackageEvent_Raw(this, &FConcertSessionPackageViewerController::GetPackageEvent)
+		.GetSizeOfPackageActivity_Raw(this, &FConcertSessionPackageViewerController::GetSizeOfPackageEvent);
 }
 
 TOptional<FConcertClientInfo> FConcertSessionPackageViewerController::GetClientInfo(FGuid ClientId) const
@@ -95,6 +96,27 @@ bool FConcertSessionPackageViewerController::GetPackageEvent(const FConcertSessi
 	return false;
 }
 
+TOptional<int64> FConcertSessionPackageViewerController::GetSizeOfPackageEvent(const FConcertSessionActivity& Activity) const
+{
+	const TOptional<FConcertSyncSessionDatabaseNonNullPtr> Database = SyncServer->GetLiveSessionDatabase(InspectedSession->GetId());
+	if (!Database)
+	{
+		return {};
+	}
+
+	const TSet<EConcertPackageUpdateType> PackageTypesWithSize { EConcertPackageUpdateType::Added, EConcertPackageUpdateType::Renamed, EConcertPackageUpdateType::Saved };
+	FConcertSyncPackageEventMetaData PackageEventMetaData;
+	if (Database->GetPackageEventMetaData(Activity.Activity.EventId, PackageEventMetaData.PackageRevision, PackageEventMetaData.PackageInfo)
+		; PackageTypesWithSize.Contains(PackageEventMetaData.PackageInfo.PackageUpdateType))
+	{
+		const FName PackageName = PackageEventMetaData.PackageInfo.PackageUpdateType == EConcertPackageUpdateType::Renamed
+			?  PackageEventMetaData.PackageInfo.NewPackageName
+			: PackageEventMetaData.PackageInfo.PackageName;
+		return Database->GetPackageSizeForRevision(PackageName);
+	}
+	return {};
+}
+
 void FConcertSessionPackageViewerController::OnSessionProduced(const FConcertSyncActivity& ProducedActivity) const
 {
 	if (ProducedActivity.EventType == EConcertSyncActivityEventType::Package)
diff --git a/Engine/Plugins/Developer/Concert/ConcertApp/MultiUserServer/Source/MultiUserServer/Private/Widgets/SessionTabs/Live/PackageViewer/ConcertSessionPackageViewerController.h b/Engine/Plugins/Developer/Concert/ConcertApp/MultiUserServer/Source/MultiUserServer/Private/Widgets/SessionTabs/Live/PackageViewer/ConcertSessionPackageViewerController.h
--- a/Engine/Plugins/Developer/Concert/ConcertApp/MultiUserServer/Source/MultiUserServer/Private/Widgets/SessionTabs/Live/PackageViewer/ConcertSessionPackageViewerController.h
+++ b/Engine/Plugins/Developer/Concert/ConcertApp/MultiUserServer/Source/MultiUserServer/Private/Widgets/SessionTabs/Live/PackageViewer/ConcertSessionPackageViewerController.h
@@ -35,5 +35,6 @@ private:
 	
 	TOptional<FConcertClientInfo> GetClientInfo(FGuid ClientId) const;
 	bool GetPackageEvent(const FConcertSessionActivity& Activity, FConcertSyncPackageEventMetaData& OutPackageEvent) const;
+	TOptional<int64> GetSizeOfPackageEvent(const FConcertSessionActivity& Activity) const;
 	void OnSessionProduced(const FConcertSyncActivity& ProducedActivity) const;
 };
diff --git a/Engine/Plugins/Developer/Concert/ConcertApp/MultiUserServer/Source/MultiUserServer/Private/Widgets/SessionTabs/Live/PackageViewer/PackageViewerColumns.cpp b/Engine/Plugins/Developer/Concert/ConcertApp/MultiUserServer/Source/MultiUserServer/Private/Widgets/SessionTabs/Live/PackageViewer/PackageViewerColumns.cpp
new file mode 100644
--- /dev/null
+++ b/Engine/Plugins/Developer/Concert/ConcertApp/MultiUserServer/Source/MultiUserServer/Private/Widgets/SessionTabs/Live/PackageViewer/PackageViewerColumns.cpp
@@ -0,0 +1,86 @@
+﻿// Copyright Epic Games, Inc. All Rights Reserved.
+
+#include "PackageViewerColumns.h"
+
+#include "Math/UnitConversion.h"
+#include "Widgets/Text/STextBlock.h"
+
+#define LOCTEXT_NAMESPACE "UnrealMultiUserUI"
+
+const FName UE::MultiUserServer::PackageViewerColumns::SizeColumnId		= TEXT("SizeColumnId");
+const FName UE::MultiUserServer::PackageViewerColumns::VersionColumnId	= TEXT("VersionColumnId");
+
+namespace UE::MultiUserServer::PackageViewerColumns
+{
+	using FFormatString = TFunction<FString(int64 Number)>;
+	
+	static FActivityColumn CreateNumericColumn(
+		FName ColumnId,
+		FText Label,
+		int32 Width,
+		EPredefinedPackageColumnOrder SortOrder,
+		FGetNumericValueFromPackageActivity GetterFunc,
+		FFormatString FormatFunc)
+	{
+		return FActivityColumn(
+			SHeaderRow::Column(ColumnId)
+				.DefaultLabel(Label)
+				.ManualWidth(Width)
+			)
+			.ColumnSortOrder(static_cast<int32>(SortOrder))
+			.GenerateColumnWidget(FActivityColumn::FGenerateColumnWidget::CreateLambda([GetterFunc, FormatFunc](const TSharedRef<SConcertSessionActivities>& Owner, const TSharedRef<FConcertSessionActivity>& Activity, SOverlay::FScopedWidgetSlotArguments& Slot)
+			{
+				Slot
+				.HAlign(HAlign_Center)
+				.Padding(1)
+				[
+					SNew(STextBlock)
+					.Text_Lambda([GetterFunc, FormatFunc, WeakActivity = TWeakPtr<FConcertSessionActivity>(Activity)]()
+					{
+						const TSharedPtr<FConcertSessionActivity> PinnedActivity = WeakActivity.Pin();
+						check(PinnedActivity);
+						const TOptional<int64> Size = GetterFunc.Execute(*PinnedActivity.Get());
+						return FText::FromString(Size ? FormatFunc(*Size) : TEXT("n/a"));
+					})
+					.HighlightText(Owner->GetHighlightText())
+				];
+			}))
+			.PopulateSearchString(FActivityColumn::FPopulateSearchString::CreateLambda([GetterFunc, FormatFunc](const TSharedRef<SConcertSessionActivities>& Owner, const FConcertSessionActivity& Activity, TArray<FString>& OutSearchStrings)
+			{
+				const TOptional<int64> Size = GetterFunc.Execute(Activity);
+				const FString Text = Size ? FormatFunc(*Size) : FString(TEXT("n/a"));
+				OutSearchStrings.Add(Text);
+			}));
+	}
+}
+
+FActivityColumn UE::MultiUserServer::PackageViewerColumns::SizeColumn(FGetSizeOfPackageActivity GetEventDataFunc)
+{
+	return CreateNumericColumn(
+		SizeColumnId,
+		LOCTEXT("Size", "Size"),
+		80,
+		EPredefinedPackageColumnOrder::Size,
+		GetEventDataFunc,
+		[](int64 Number)
+		{
+			const FNumericUnit<int64> Unit = FUnitConversion::QuantizeUnitsToBestFit(Number, EUnit::Bytes);
+			return FString::Printf(TEXT("%lld %s"), Unit.Value, FUnitConversion::GetUnitDisplayString(Unit.Units));
+		});
+}
+
+FActivityColumn UE::MultiUserServer::PackageViewerColumns::VersionColumn(FGetVersionOfPackageActivity GetEventDataFunc)
+{
+	return CreateNumericColumn(
+		VersionColumnId,
+		LOCTEXT("Version", "Version"),
+		80,
+		EPredefinedPackageColumnOrder::Version,
+		GetEventDataFunc,
+		[](int64 Number)
+		{
+			return FString::Printf(TEXT("%lld"), Number);
+		});
+}
+
+#undef LOCTEXT_NAMESPACE 
\ No newline at end of file
diff --git a/Engine/Plugins/Developer/Concert/ConcertApp/MultiUserServer/Source/MultiUserServer/Private/Widgets/SessionTabs/Live/PackageViewer/PackageViewerColumns.h b/Engine/Plugins/Developer/Concert/ConcertApp/MultiUserServer/Source/MultiUserServer/Private/Widgets/SessionTabs/Live/PackageViewer/PackageViewerColumns.h
new file mode 100644
--- /dev/null
+++ b/Engine/Plugins/Developer/Concert/ConcertApp/MultiUserServer/Source/MultiUserServer/Private/Widgets/SessionTabs/Live/PackageViewer/PackageViewerColumns.h
@@ -0,0 +1,26 @@
+﻿// Copyright Epic Games, Inc. All Rights Reserved.
+
+#pragma once
+
+#include "CoreMinimal.h"
+#include "Session/Activity/ActivityColumn.h"
+#include "Session/Activity/PredefinedActivityColumns.h"
+
+namespace UE::MultiUserServer::PackageViewerColumns
+{
+	DECLARE_DELEGATE_RetVal_OneParam(TOptional<int64>, FGetNumericValueFromPackageActivity, const FConcertSessionActivity& /*Activity*/);
+	using FGetSizeOfPackageActivity = FGetNumericValueFromPackageActivity;
+	using FGetVersionOfPackageActivity = FGetNumericValueFromPackageActivity;
+	
+	enum class EPredefinedPackageColumnOrder : int32
+	{
+		Size = static_cast<int32>(ConcertSharedSlate::ActivityColumn::EPredefinedColumnOrder::ClientName) + 1,
+		Version = Size + 1
+	};
+	
+	extern const FName SizeColumnId;
+	extern const FName VersionColumnId;
+	
+	FActivityColumn SizeColumn(FGetSizeOfPackageActivity GetEventDataFunc);
+	FActivityColumn VersionColumn(FGetVersionOfPackageActivity GetEventDataFunc);
+};
diff --git a/Engine/Plugins/Developer/Concert/ConcertApp/MultiUserServer/Source/MultiUserServer/Private/Widgets/SessionTabs/Live/PackageViewer/SConcertSessionPackageViewer.cpp b/Engine/Plugins/Developer/Concert/ConcertApp/MultiUserServer/Source/MultiUserServer/Private/Widgets/SessionTabs/Live/PackageViewer/SConcertSessionPackageViewer.cpp
--- a/Engine/Plugins/Developer/Concert/ConcertApp/MultiUserServer/Source/MultiUserServer/Private/Widgets/SessionTabs/Live/PackageViewer/SConcertSessionPackageViewer.cpp
+++ b/Engine/Plugins/Developer/Concert/ConcertApp/MultiUserServer/Source/MultiUserServer/Private/Widgets/SessionTabs/Live/PackageViewer/SConcertSessionPackageViewer.cpp
@@ -2,6 +2,8 @@
 
 #include "SConcertSessionPackageViewer.h"
 
+#include "PackageViewerColumns.h"
+
 #include "Session/Activity/PredefinedActivityColumns.h"
 #include "Session/Activity/SConcertSessionActivities.h"
 #include "Widgets/Input/SSearchBox.h"
@@ -11,6 +13,9 @@
 void SConcertSessionPackageViewer::Construct(const FArguments& InArgs)
 {
 	using namespace UE::ConcertSharedSlate;
+	using namespace UE::MultiUserServer;
+
+	check(InArgs._GetSizeOfPackageActivity.IsBound());
 	
 	ActivityListViewOptions = MakeShared<FConcertSessionActivitiesOptions>();
 	ActivityListViewOptions->bEnableConnectionActivityFiltering = false;
@@ -23,7 +28,8 @@ void SConcertSessionPackageViewer::Construct(const FArguments& InArgs)
 		.OnMapActivityToClient(InArgs._GetClientInfo)
 		.HighlightText(this, &SConcertSessionPackageViewer::HighlightSearchedText)
 		.Columns({
-			// TODO: Create Icon (depending on EConcertPackageUpdateType), Size, and Version
+			PackageViewerColumns::SizeColumn(InArgs._GetSizeOfPackageActivity),
+			PackageViewerColumns::VersionColumn(PackageViewerColumns::FGetVersionOfPackageActivity::CreateSP(this, &SConcertSessionPackageViewer::GetVersionOfPackageActivity, InArgs._GetPackageEvent))
 		})
 		.TimeFormat(ActivityListViewOptions.Get(), &FConcertSessionActivitiesOptions::GetTimeFormat)
 		.ConnectionActivitiesVisibility(EVisibility::Collapsed)
@@ -74,6 +80,19 @@ void SConcertSessionPackageViewer::AppendActivity(FConcertSessionActivity Activi
 	ActivityListView->Append(NewActivity);
 }
 
+TOptional<int64> SConcertSessionPackageViewer::GetVersionOfPackageActivity(const FConcertSessionActivity& Activity, SConcertSessionActivities::FGetPackageEvent GetPackageEventFunc) const
+{
+	FConcertSyncPackageEventMetaData PackageEventMetaData;
+	if (GetPackageEventFunc.Execute(Activity, PackageEventMetaData))
+	{
+		const TSet<EConcertPackageUpdateType> NoRevisionTypes { EConcertPackageUpdateType::Deleted, EConcertPackageUpdateType::Dummy };
+		return NoRevisionTypes.Contains(PackageEventMetaData.PackageInfo.PackageUpdateType)
+			? TOptional<int64>{}
+			: PackageEventMetaData.PackageRevision;
+	}
+	return {};
+}
+
 void SConcertSessionPackageViewer::OnSearchTextChanged(const FText& InSearchText)
 {
 	SearchedText = InSearchText;
diff --git a/Engine/Plugins/Developer/Concert/ConcertApp/MultiUserServer/Source/MultiUserServer/Private/Widgets/SessionTabs/Live/PackageViewer/SConcertSessionPackageViewer.h b/Engine/Plugins/Developer/Concert/ConcertApp/MultiUserServer/Source/MultiUserServer/Private/Widgets/SessionTabs/Live/PackageViewer/SConcertSessionPackageViewer.h
--- a/Engine/Plugins/Developer/Concert/ConcertApp/MultiUserServer/Source/MultiUserServer/Private/Widgets/SessionTabs/Live/PackageViewer/SConcertSessionPackageViewer.h
+++ b/Engine/Plugins/Developer/Concert/ConcertApp/MultiUserServer/Source/MultiUserServer/Private/Widgets/SessionTabs/Live/PackageViewer/SConcertSessionPackageViewer.h
@@ -3,6 +3,8 @@
 #pragma once
 
 #include "CoreMinimal.h"
+#include "PackageViewerColumns.h"
+
 #include "Session/Activity/SConcertSessionActivities.h"
 #include "Widgets/SCompoundWidget.h"
 #include "Widgets/Input/SSearchBox.h"
@@ -14,16 +16,11 @@ class FConcertSessionActivitiesOptions;
 class SConcertSessionPackageViewer : public SCompoundWidget
 {
 public:
-
-	struct FConcertPackageActivity
-	{
-		FConcertSyncActivity Base;
-		FConcertSyncPackageEventData PackageEvent;
-	};
 	
 	SLATE_BEGIN_ARGS(SConcertSessionPackageViewer) {}
 		SLATE_EVENT(SConcertSessionActivities::FGetPackageEvent, GetPackageEvent)
 		SLATE_EVENT(SConcertSessionActivities::FGetActivityClientInfoFunc, GetClientInfo)
+		SLATE_EVENT(UE::MultiUserServer::PackageViewerColumns::FGetSizeOfPackageActivity, GetSizeOfPackageActivity)
 	SLATE_END_ARGS()
 
 	void Construct(const FArguments& InArgs);
@@ -44,10 +41,12 @@ private:
 	/** The searched text to highlight. */
 	FText SearchedText;
 
+	TOptional<int64> GetVersionOfPackageActivity(const FConcertSessionActivity& Activity, SConcertSessionActivities::FGetPackageEvent GetPackageEventFunc) const;
+
 	/** Invoked when the text in the search box widget changes. */
 	void OnSearchTextChanged(const FText& InSearchText);
 	/** Invoked when the text in the search box widget is committed. */
 	void OnSearchTextCommitted(const FText& InFilterText, ETextCommit::Type CommitType);
 	/** Returns the text to highlight when the search bar has a text set. */
 	FText HighlightSearchedText() const;
-};
+};
\ No newline at end of file
diff --git a/Engine/Plugins/Developer/Concert/ConcertSync/ConcertSyncCore/Source/ConcertSyncCore/Private/ConcertSyncSessionDatabase.cpp b/Engine/Plugins/Developer/Concert/ConcertSync/ConcertSyncCore/Source/ConcertSyncCore/Private/ConcertSyncSessionDatabase.cpp
--- a/Engine/Plugins/Developer/Concert/ConcertSync/ConcertSyncCore/Source/ConcertSyncCore/Private/ConcertSyncSessionDatabase.cpp
+++ b/Engine/Plugins/Developer/Concert/ConcertSync/ConcertSyncCore/Source/ConcertSyncCore/Private/ConcertSyncSessionDatabase.cpp
@@ -2549,16 +2549,8 @@ bool FConcertSyncSessionDatabase::AddPersistEventForHeadRevision(FName InPackage
 
 bool FConcertSyncSessionDatabase::GetPackageInfoForRevision(const FName InPackageName, FConcertPackageInfo& OutPackageInfo, const int64* InPackageRevision) const
 {
-	int64 PackageRevision = 0;
-	if (InPackageRevision)
-	{
-		PackageRevision = *InPackageRevision;
-	}
-	else if (!GetPackageHeadRevision(InPackageName, PackageRevision))
-	{
-		return false;
-	}
-	if (PackageRevision == 0)
+	const TOptional<int64> PackageRevision = GetSpecifiedOrHeadPackageRevision(InPackageName, InPackageRevision);
+	if (!PackageRevision)
 	{
 		return false;
 	}
@@ -2570,43 +2562,35 @@ bool FConcertSyncSessionDatabase::GetPackageInfoForRevision(const FName InPackag
 	}
 
 	FString DataFilename;
-	return Statements->GetPackageDataForRevision(PackageNameId, PackageRevision, OutPackageInfo, DataFilename);
+	return Statements->GetPackageDataForRevision(PackageNameId, *PackageRevision, OutPackageInfo, DataFilename);
 }
 
 bool FConcertSyncSessionDatabase::GetPackageDataForRevision(const FName InPackageName, const TFunctionRef<void(const FConcertPackageInfo&, FConcertPackageDataStream&)>& InPackageDataCallback, const int64* InPackageRevision) const
 {
-	int64 PackageRevision = 0;
-	if (InPackageRevision)
-	{
-		PackageRevision = *InPackageRevision;
-	}
-	else if (!GetPackageHeadRevision(InPackageName, PackageRevision))
-	{
-		return false;
-	}
-	if (PackageRevision == 0)
-	{
-		return false;
-	}
-
-	int64 PackageNameId = 0;
-	if (!GetPackageNameId(InPackageName, PackageNameId))
-	{
-		return false;
-	}
-
-	FConcertPackageInfo PackageInfo;
-	FString DataFilename;
-	if (Statements->GetPackageDataForRevision(PackageNameId, PackageRevision, PackageInfo, DataFilename))
+	return HandleRequestPackageRequest(InPackageName, InPackageRevision, [this, InPackageDataCallback](FConcertPackageInfo PackageInfo, FString DataFilename)
 	{
 		const FString PackageDataPathname = PackageDataUtil::GetDataPath(SessionPath) / DataFilename;
-		return LoadPackage(PackageDataPathname, [&PackageInfo, &InPackageDataCallback](FConcertPackageDataStream& InPackageDataStream)
+		return LoadPackage(PackageDataPathname, [&PackageInfo, InPackageDataCallback](FConcertPackageDataStream& InPackageDataStream)
 		{
 			InPackageDataCallback(PackageInfo, InPackageDataStream);
 		});
-	}
+	});
+}
 
-	return false;
+TOptional<int64> FConcertSyncSessionDatabase::GetPackageSizeForRevision(const FName InPackageName, const int64* InPackageRevision) const
+{
+	TOptional<int64> Result;
+	HandleRequestPackageRequest(InPackageName, InPackageRevision, [this, &Result](FConcertPackageInfo PackageInfo, FString DataFilename)
+	{
+		const FString PackageDataPathname = PackageDataUtil::GetDataPath(SessionPath) / DataFilename;
+		const int64 FileSize = IFileManager::Get().FileSize(*PackageDataPathname);
+		if (FileSize >= 0)
+		{
+			Result = FileSize;
+		}
+		return true;
+	});
+	return Result;
 }
 
 bool FConcertSyncSessionDatabase::GetPackageHeadRevision(const FName InPackageName, int64& OutRevision) const
@@ -2839,6 +2823,42 @@ void FConcertSyncSessionDatabase::UpdateAsynchronousTasks()
 	}
 }
 
+bool FConcertSyncSessionDatabase::HandleRequestPackageRequest(const FName InPackageName, const int64* InPackageRevision, FProcessPackageRequest HandleFunc) const
+{
+	int64 PackageNameId = 0;
+	FConcertPackageInfo PackageInfo;
+	FString DataFilename;
+	
+	const TOptional<int64> PackageRevision = GetSpecifiedOrHeadPackageRevision(InPackageName, InPackageRevision);
+	if (PackageRevision
+		&& GetPackageNameId(InPackageName, PackageNameId)
+		&& Statements->GetPackageDataForRevision(PackageNameId, *PackageRevision, PackageInfo, DataFilename))
+	{
+		return HandleFunc(PackageInfo, DataFilename);
+	}
+	
+	return false;
+}
+
+TOptional<int64> FConcertSyncSessionDatabase::GetSpecifiedOrHeadPackageRevision(FName InPackageName, const int64* InPackageRevision) const
+{
+	int64 PackageRevision = 0;
+	if (InPackageRevision)
+	{
+		PackageRevision = *InPackageRevision;
+	}
+	else if (!GetPackageHeadRevision(InPackageName, PackageRevision))
+	{
+		return {};
+	}
+	
+	if (PackageRevision == 0)
+	{
+		return {};
+	}
+	return PackageRevision;
+}
+
 void FConcertSyncSessionDatabase::ScheduleAsyncWrite(const FString& InDstPackageBlobPathname, FConcertPackageDataStream& InPackageDataStream)
 {
 	check(InPackageDataStream.DataBlob && InPackageDataStream.DataBlob->Num() > 0);
diff --git a/Engine/Plugins/Developer/Concert/ConcertSync/ConcertSyncCore/Source/ConcertSyncCore/Public/ConcertSyncSessionDatabase.h b/Engine/Plugins/Developer/Concert/ConcertSync/ConcertSyncCore/Source/ConcertSyncCore/Public/ConcertSyncSessionDatabase.h
--- a/Engine/Plugins/Developer/Concert/ConcertSync/ConcertSyncCore/Source/ConcertSyncCore/Public/ConcertSyncSessionDatabase.h
+++ b/Engine/Plugins/Developer/Concert/ConcertSync/ConcertSyncCore/Source/ConcertSyncCore/Public/ConcertSyncSessionDatabase.h
@@ -619,6 +619,16 @@ public:
 	 * @return True if package data could be found for the given revision, false otherwise.
 	 */
 	bool GetPackageDataForRevision(const FName InPackageName, const TFunctionRef<void(const FConcertPackageInfo&, FConcertPackageDataStream&)>& InCallback, const int64* InPackageRevision = nullptr) const;
+
+	/**
+	 * Gets the package size in bytes of a package
+	 *
+	 * @param InPackageName				The name of the package to get the head revision for.
+	 * @param InPackageRevision			The revision of the package to get the data for, or null to get the head revision.
+	 *
+	 * @return True if package data could be found for the given revision, false otherwise.
+	 */
+	TOptional<int64> GetPackageSizeForRevision(const FName InPackageName, const int64* InPackageRevision = nullptr) const;
 	
 	/**
 	 * Get the head revision in this database for the given package name.
@@ -681,6 +691,14 @@ public:
 	FOnActivityProduced& OnActivityProduced() { return ActivityProducedEvent; }
 	
 private:
+
+	using FProcessPackageRequest = TFunctionRef<bool(const FConcertPackageInfo& PackageInfo, const FString& DataFilename)>;
+	/** Helper function which obtains package information and passes it to HandleFunc. */
+	bool HandleRequestPackageRequest(const FName InPackageName, const int64* InPackageRevision, FProcessPackageRequest HandleFunc) const;
+
+	/** Helper functions that for getting a package revision for an optional package revision argument */
+	TOptional<int64> GetSpecifiedOrHeadPackageRevision(FName InPackageName, const int64* InPackageRevision) const;
+	
 	/**
 	 * Schedule an asynchronous write for the given Package Stream.  The stream must be in-memory. File sharing
 	 * asynchronous write is not supported.
diff --git a/Engine/Plugins/Developer/Concert/ConcertUI/ConcertSharedSlate/Source/ConcertSharedSlate/Private/Session/Activity/SConcertSessionActivities.cpp b/Engine/Plugins/Developer/Concert/ConcertUI/ConcertSharedSlate/Source/ConcertSharedSlate/Private/Session/Activity/SConcertSessionActivities.cpp
--- a/Engine/Plugins/Developer/Concert/ConcertUI/ConcertSharedSlate/Source/ConcertSharedSlate/Private/Session/Activity/SConcertSessionActivities.cpp
+++ b/Engine/Plugins/Developer/Concert/ConcertUI/ConcertSharedSlate/Source/ConcertSharedSlate/Private/Session/Activity/SConcertSessionActivities.cpp
@@ -366,8 +366,12 @@ TSharedRef<SHeaderRow> SConcertSessionActivities::CreateHeaderRow(const FArgumen
 	Columns.Add(ActivityColumn::Summary());
 	
 	HeaderRow = SNew(SHeaderRow);
+	TSet<FName> DuplicateColumnDetection;
 	for (FActivityColumn& Column : Columns)
 	{
+		check(!DuplicateColumnDetection.Contains(Column.ColumnId));
+		DuplicateColumnDetection.Add(Column.ColumnId);
+		
 		// SHeaderRow owns the columns and deletes them when destroyed
 		FActivityColumn* ManagedByHeaderRow = new FActivityColumn(MoveTemp(Column));
 		HeaderRow->AddColumn(*ManagedByHeaderRow);
