Analyse removalability of activities

diff --git a/Engine/Plugins/Developer/Concert/ConcertMain/Source/Concert/Public/ConcertMessageData.h b/Engine/Plugins/Developer/Concert/ConcertMain/Source/Concert/Public/ConcertMessageData.h
--- a/Engine/Plugins/Developer/Concert/ConcertMain/Source/Concert/Public/ConcertMessageData.h
+++ b/Engine/Plugins/Developer/Concert/ConcertMain/Source/Concert/Public/ConcertMessageData.h
@@ -11,6 +11,8 @@
 
 class FStructOnScope;
 
+using FActivityID = int64;
+
 UENUM()
 enum class EConcertServerFlags : uint8
 {
diff --git a/Engine/Plugins/Developer/Concert/ConcertSync/ConcertSyncCore/Source/ConcertSyncCore/Private/HistoryEdition/ActivityDependencyEdge.cpp b/Engine/Plugins/Developer/Concert/ConcertSync/ConcertSyncCore/Source/ConcertSyncCore/Private/HistoryEdition/ActivityDependencyEdge.cpp
new file mode 100644
--- /dev/null
+++ b/Engine/Plugins/Developer/Concert/ConcertSync/ConcertSyncCore/Source/ConcertSyncCore/Private/HistoryEdition/ActivityDependencyEdge.cpp
@@ -0,0 +1,19 @@
+// Copyright Epic Games, Inc. All Rights Reserved.
+
+#include "HistoryEdition/ActivityDependencyEdge.h"
+
+FString UE::ConcertSyncCore::LexToString(EActivityDependencyReason Reason)
+{
+	switch (Reason)
+	{
+		case EActivityDependencyReason::PackageCreation: return TEXT("PackageCreation");
+		case EActivityDependencyReason::PackageRemoval: return TEXT("PackageRemoval");
+		case EActivityDependencyReason::PackageRename: return TEXT("PackageRename");
+		case EActivityDependencyReason::EditAfterPreviousPackageEdit: return TEXT("EditPossiblyDependsOnPackage");
+		case EActivityDependencyReason::SubobjectCreation: return TEXT("SubobjectCreation");
+		case EActivityDependencyReason::SubobjectRemoval: return TEXT("SubobjectRemoval");
+		default:
+			checkNoEntry();
+			return TEXT("Invalid");
+	}
+}
diff --git a/Engine/Plugins/Developer/Concert/ConcertSync/ConcertSyncCore/Source/ConcertSyncCore/Private/HistoryEdition/ActivityDependencyGraph.cpp b/Engine/Plugins/Developer/Concert/ConcertSync/ConcertSyncCore/Source/ConcertSyncCore/Private/HistoryEdition/ActivityDependencyGraph.cpp
new file mode 100644
--- /dev/null
+++ b/Engine/Plugins/Developer/Concert/ConcertSync/ConcertSyncCore/Source/ConcertSyncCore/Private/HistoryEdition/ActivityDependencyGraph.cpp
@@ -0,0 +1,76 @@
+﻿// Copyright Epic Games, Inc. All Rights Reserved.
+
+#include "HistoryEdition/ActivityDependencyGraph.h"
+
+#include "ConcertLogGlobal.h"
+#include "HistoryEdition/ActivityDependencyEdge.h"
+#include "HistoryEdition/ActivityNode.h"
+
+#include "Algo/ForEach.h"
+
+UE::ConcertSyncCore::FActivityNodeID UE::ConcertSyncCore::FActivityDependencyGraph::AddActivity(int64 ActivityIndex)
+{
+	const TOptional<FActivityNodeID> PreexistingID = FindNodeByActivity(ActivityIndex);
+	if (ensureMsgf(!PreexistingID.IsSet(), TEXT("Activity is already registered!")))
+	{
+		FActivityNodeID NodeID(Nodes.Num());
+		Nodes.Emplace(ActivityIndex, NodeID);
+		return NodeID;
+	}
+
+	return *PreexistingID;
+}
+
+bool UE::ConcertSyncCore::FActivityDependencyGraph::AddDependency(FActivityNodeID From, FActivityDependencyEdge To)
+{
+	if (!ensure(Nodes.IsValidIndex(From.ID) && Nodes.IsValidIndex(To.GetDependedOnNodeID().ID)))
+	{
+		return false;
+	}
+	
+	FActivityNode& FromNode = GetNodeByIdInternal(From);
+	FActivityNode& ToNode = GetNodeByIdInternal(To.GetDependedOnNodeID());
+
+	const bool bIsValidDependency = FromNode.GetActivityId() > ToNode.GetActivityId();
+	if (bIsValidDependency)
+	{
+		FromNode.Dependencies.AddUnique(To);
+		ToNode.AffectedChildren.AddUnique(From);
+		return true;
+	}
+	
+	UE_CLOG(!bIsValidDependency, LogConcert, Error, TEXT("Activities can only depend on preceding activities!"));
+	return false;
+}
+
+TOptional<UE::ConcertSyncCore::FActivityNodeID> UE::ConcertSyncCore::FActivityDependencyGraph::FindNodeByActivity(int64 ActivityID) const
+{
+	for (int32 i = 0; i < Nodes.Num(); ++i)
+	{
+		if (Nodes[i].ActivityId == ActivityID)
+		{
+			return FActivityNodeID(i);
+		}
+	}
+
+	return {};
+}
+
+const UE::ConcertSyncCore::FActivityNode& UE::ConcertSyncCore::FActivityDependencyGraph::GetNodeById(FActivityNodeID ID) const
+{
+	const size_t Index = static_cast<size_t>(ID);
+	check(Nodes.IsValidIndex(Index));
+	return Nodes[Index];
+}
+
+UE::ConcertSyncCore::FActivityNode& UE::ConcertSyncCore::FActivityDependencyGraph::GetNodeByIdInternal(FActivityNodeID ID)
+{
+	const size_t Index = static_cast<size_t>(ID);
+	check(Nodes.IsValidIndex(Index));
+	return Nodes[Index];
+}
+
+void UE::ConcertSyncCore::FActivityDependencyGraph::ForEachNode(TFunctionRef<void(const FActivityNode&)> ConsumerFunc) const
+{
+	Algo::ForEach(Nodes, ConsumerFunc);
+}
diff --git a/Engine/Plugins/Developer/Concert/ConcertSync/ConcertSyncCore/Source/ConcertSyncCore/Private/HistoryEdition/ActivityNode.cpp b/Engine/Plugins/Developer/Concert/ConcertSync/ConcertSyncCore/Source/ConcertSyncCore/Private/HistoryEdition/ActivityNode.cpp
new file mode 100644
--- /dev/null
+++ b/Engine/Plugins/Developer/Concert/ConcertSync/ConcertSyncCore/Source/ConcertSyncCore/Private/HistoryEdition/ActivityNode.cpp
@@ -0,0 +1,33 @@
+// Copyright Epic Games, Inc. All Rights Reserved.
+
+#include "HistoryEdition/ActivityNode.h"
+
+#include "HistoryEdition/ActivityDependencyEdge.h"
+#include "HistoryEdition/ActivityDependencyGraph.h"
+
+bool UE::ConcertSyncCore::FActivityNode::DependsOnNode(FActivityNodeID NodeId, TOptional<EActivityDependencyReason> WithReason, TOptional<EDependencyStrength> WithStrength) const
+{
+	return Dependencies.ContainsByPredicate([NodeId, WithReason, WithStrength](const FActivityDependencyEdge& Edge)
+	{
+		return Edge.GetDependedOnNodeID() == NodeId
+			&& (!WithReason.IsSet() || *WithReason == Edge.GetDependencyReason())
+			&& (!WithStrength.IsSet() || *WithStrength == Edge.GetDependencyStrength());
+	});
+}
+
+bool UE::ConcertSyncCore::FActivityNode::DependsOnActivity(int64 InActivityId, const FActivityDependencyGraph& Graph, TOptional<EActivityDependencyReason> WithReason, TOptional<EDependencyStrength> WithStrength) const
+{
+	const TOptional<FActivityNodeID> ActivityNodeId = Graph.FindNodeByActivity(InActivityId);
+	return ActivityNodeId && Dependencies.ContainsByPredicate([ActivityNodeId, WithReason, WithStrength](const FActivityDependencyEdge& Edge)
+	{
+		return Edge.GetDependedOnNodeID() == *ActivityNodeId
+			&& (!WithReason.IsSet() || *WithReason == Edge.GetDependencyReason())
+			&& (!WithStrength.IsSet() || *WithStrength == Edge.GetDependencyStrength());
+	});
+}
+
+bool UE::ConcertSyncCore::FActivityNode::AffectsActivity(int64 InActivityId, const FActivityDependencyGraph& Graph) const
+{
+	const TOptional<FActivityNodeID> ActivityNodeId = Graph.FindNodeByActivity(InActivityId);
+	return ActivityNodeId && AffectedChildren.Contains(*ActivityNodeId);
+}
\ No newline at end of file
diff --git a/Engine/Plugins/Developer/Concert/ConcertSync/ConcertSyncCore/Source/ConcertSyncCore/Private/HistoryEdition/DebugDependencyGraph.cpp b/Engine/Plugins/Developer/Concert/ConcertSync/ConcertSyncCore/Source/ConcertSyncCore/Private/HistoryEdition/DebugDependencyGraph.cpp
new file mode 100644
--- /dev/null
+++ b/Engine/Plugins/Developer/Concert/ConcertSync/ConcertSyncCore/Source/ConcertSyncCore/Private/HistoryEdition/DebugDependencyGraph.cpp
@@ -0,0 +1,168 @@
+﻿// Copyright Epic Games, Inc. All Rights Reserved.
+
+#include "HistoryEdition/DebugDependencyGraph.h"
+
+#include "HistoryEdition/ActivityNode.h"
+#include "HistoryEdition/ActivityDependencyEdge.h"
+
+#include "ConcertLogGlobal.h"
+#include "ConcertSyncSessionDatabase.h"
+
+void UE::ConcertSyncCore::Graphviz::MakeNodeTitle(FGraphStringBuilder& WriteTo, FActivityNodeID ToStringify, const FActivityDependencyGraph& Graph, const FConcertSyncSessionDatabase& SessionDatabase, ENodeTitleFlags NodeTitleFlags)
+{
+	const FActivityNode& Node = Graph.GetNodeById(ToStringify);
+	FConcertSyncActivity Activity;
+	SessionDatabase.GetActivity(Node.GetActivityId(), Activity);
+		
+	if ((NodeTitleFlags & ENodeTitleFlags::Summary) != ENodeTitleFlags::None)
+	{
+		FConcertSyncTransactionEvent TransactionEvent;
+		switch (Activity.EventType)
+		{
+		case EConcertSyncActivityEventType::Transaction:
+			if (SessionDatabase.GetTransactionEvent(Activity.EventId, TransactionEvent))
+			{
+				WriteTo << FConcertSyncTransactionActivitySummary::CreateSummaryForEvent(TransactionEvent).ToDisplayText(FText::GetEmpty()).ToString();
+			}
+			break;
+		case EConcertSyncActivityEventType::Package:
+			SessionDatabase.GetPackageEvent(Activity.EventId, [&WriteTo](FConcertSyncPackageEventData& Event)
+			{
+				WriteTo << FConcertSyncPackageActivitySummary::CreateSummaryForEvent(Event.MetaData.PackageInfo).ToDisplayText(FText::GetEmpty()).ToString();
+			});
+			break;
+
+			// Dependency graph ignores ignores these
+			case EConcertSyncActivityEventType::None:
+			case EConcertSyncActivityEventType::Connection:
+			case EConcertSyncActivityEventType::Lock:
+			default:
+				checkNoEntry();
+			break;
+		}
+	}
+
+	if ((NodeTitleFlags & ENodeTitleFlags::ActivityID) != ENodeTitleFlags::None)
+	{
+		WriteTo << FString::Printf(TEXT(" (%lld)"), Activity.ActivityId);
+	}
+	return;
+}
+
+void UE::ConcertSyncCore::Graphviz::GetNodeStyle(FGraphStringBuilder& WriteTo, FActivityNodeID ToStringify, const FActivityDependencyGraph& Graph, const FConcertSyncSessionDatabase& SessionDatabase)
+{
+	// https://graphviz.org/doc/info/shapes.html
+	const TMap<EConcertSyncActivityEventType, FString> Shapes {
+		{ EConcertSyncActivityEventType::None, TEXT("plaintext") },
+		{ EConcertSyncActivityEventType::Connection, TEXT("plaintext") },
+		{ EConcertSyncActivityEventType::Lock, TEXT("plaintext") },
+		{ EConcertSyncActivityEventType::Transaction, TEXT("ellipse") },
+		{ EConcertSyncActivityEventType::Package, TEXT("folder") }
+	};
+	const TMap<EConcertPackageUpdateType, FColor> PackageColours {
+			{ EConcertPackageUpdateType::Dummy,  FColor::White },
+			{ EConcertPackageUpdateType::Added, FColor(FColor::Green.R * 0.6, FColor::Green.G * 0.6, FColor::Green.B * 0.6) },
+			{ EConcertPackageUpdateType::Saved,  FColor::Green },
+			{ EConcertPackageUpdateType::Renamed, FColor(FColor::Orange.R * 0.8, FColor::Orange.G * 0.8, FColor::Orange.B * 0.8) },
+			{ EConcertPackageUpdateType::Deleted, FColor(FColor::Red.R * 0.8, FColor::Red.G * 0.8, FColor::Red.B * 0.8) }
+	};
+
+	const FActivityNode& Node = Graph.GetNodeById(ToStringify);
+	FConcertSyncActivity Activity;
+	if (!SessionDatabase.GetActivity(Node.GetActivityId(), Activity))
+	{
+		return;
+	}
+
+	const FString Shape = Shapes[Activity.EventType];
+	FString ColorHex;
+
+	if (Activity.EventType == EConcertSyncActivityEventType::Package)
+	{
+		SessionDatabase.GetPackageEvent(Activity.EventId, [&ColorHex, &PackageColours](FConcertSyncPackageEventData& PackageEventData)
+		{
+			ColorHex = PackageColours[PackageEventData.MetaData.PackageInfo.PackageUpdateType].ToHex();
+		});
+	}
+	else
+	{
+		ColorHex = FColor(255 * 0.5f, 255 * 0.5f, 255 * 0.5f).ToHex();
+	}
+
+	WriteTo << "[shape=\"" << *Shape << "\" style=\"filled\" fillcolor=\"#" << *ColorHex << "\"]";
+}
+
+void UE::ConcertSyncCore::Graphviz::GetEdgeStyle(FGraphStringBuilder& WriteTo, const FActivityDependencyEdge& ToStringify)
+{
+	const TMap<EDependencyStrength, FString> ArrowTypes {
+			{ EDependencyStrength::HardDependency, TEXT("open") },
+			{ EDependencyStrength::PossibleDependency, TEXT("halfopen") }
+	};
+	
+	const FString ArrowHead = ArrowTypes[ToStringify.GetDependencyStrength()];
+	const FString ReasonAsString = LexToString(ToStringify.GetDependencyReason());
+	WriteTo << "[arrowhead=\"" << *ArrowHead << "\" label = \"" << *ReasonAsString << "\"]";
+}
+
+FString UE::ConcertSyncCore::Graphviz::ExportToGraphviz(const FActivityDependencyGraph& Graph, const FConcertSyncSessionDatabase& SessionDatabase, ENodeTitleFlags NodeTitleFlags)
+{
+	UE_CLOG(NodeTitleFlags == ENodeTitleFlags::None, LogConcertDebug, Warning, TEXT("No node title flags specified"));
+
+	auto MakeNodeTitle = [&Graph, &SessionDatabase](FGraphStringBuilder& WriteTo, FActivityNodeID ToStringify)
+	{
+		Graphviz::MakeNodeTitle(WriteTo, ToStringify, Graph, SessionDatabase);
+	};
+	auto GetNodeStyle = [&Graph, &SessionDatabase](FGraphStringBuilder& WriteTo, FActivityNodeID ToStringify)
+	{
+		Graphviz::GetNodeStyle(WriteTo, ToStringify, Graph, SessionDatabase);
+	};
+	auto GetEdgeStyle = [](FGraphStringBuilder& WriteTo, FActivityNodeID EdgeStart, const FActivityDependencyEdge& ToStringify)
+	{
+		Graphviz::GetEdgeStyle(WriteTo, ToStringify);
+	};
+	
+	return ExportToGraphviz(Graph, MakeNodeTitle, GetNodeStyle, GetEdgeStyle);
+}
+
+FString UE::ConcertSyncCore::Graphviz::ExportToGraphviz(const FActivityDependencyGraph& Graph, FMakeNodeTitle MakeNodeTitleFunc, FGetNodeStyle GetNodeStyleFunc, FGetEdgeStyle GetEdgeStyleFunc)
+{
+	FGraphStringBuilder Result;
+	Result << "digraph ActivityDependencyGraph {\n";
+	
+	TMap<FActivityNodeID, FString> NodeTitles;
+	auto GetTitle = [&NodeTitles, MakeNodeTitleFunc](FActivityNodeID NodeID)
+	{
+		FString* Title = NodeTitles.Find(NodeID);
+		if (!Title)
+		{
+			FGraphStringBuilder TitleBuilder;
+			TitleBuilder << "\"";
+			MakeNodeTitleFunc(TitleBuilder, NodeID);
+			TitleBuilder << "\"";
+			Title = &NodeTitles.Add(NodeID, TitleBuilder.ToString());
+		};
+		return *Title;
+	};
+	Graph.ForEachNode([&Result, GetTitle, GetNodeStyleFunc, GetEdgeStyleFunc](const FActivityNode& Node)
+	{
+		const FActivityNodeID NodeID = Node.GetNodeIndex();
+		const FString CurrentNodeTitle = *GetTitle(NodeID);
+		
+		Result << "\t" << *CurrentNodeTitle << " ";
+		GetNodeStyleFunc(Result, NodeID);
+		Result << "\n";
+		
+		for (const FActivityDependencyEdge& OutgoingDependency : Node.GetDependencies())
+		{
+			Result << "\t" << CurrentNodeTitle << " -> " << *GetTitle(OutgoingDependency.GetDependedOnNodeID()) << " ";
+			GetEdgeStyleFunc(Result, NodeID, OutgoingDependency);
+			Result << "\n"; 
+		}
+		
+		// Make blocks easier to discern for human reader
+		Result << "\n";
+	});
+
+	Result << TEXT("\n}");
+	return Result.ToString();
+}
diff --git a/Engine/Plugins/Developer/Concert/ConcertSync/ConcertSyncCore/Source/ConcertSyncCore/Private/HistoryEdition/DependencyGraphBuilder.cpp b/Engine/Plugins/Developer/Concert/ConcertSync/ConcertSyncCore/Source/ConcertSyncCore/Private/HistoryEdition/DependencyGraphBuilder.cpp
new file mode 100644
--- /dev/null
+++ b/Engine/Plugins/Developer/Concert/ConcertSync/ConcertSyncCore/Source/ConcertSyncCore/Private/HistoryEdition/DependencyGraphBuilder.cpp
@@ -0,0 +1,431 @@
+﻿// Copyright Epic Games, Inc. All Rights Reserved.
+
+#include "HistoryEdition/DependencyGraphBuilder.h"
+
+#include "ConcertSyncSessionDatabase.h"
+#include "HistoryEdition/ActivityDependencyEdge.h"
+#include "HistoryEdition/ActivityDependencyGraph.h"
+
+namespace UE::ConcertSyncCore
+{
+	/** Tracks the activities that last added, saved, removed, edited or modified a package. */
+	struct FPackageTracker
+	{
+		TMap<FName, FActivityID> AddedPackages;
+		TMap<FName, FActivityID> SavedPackages;
+		TMap<FName, FActivityID> RemovedPackages;
+		TMap<FName, FActivityID> RenamedPackages;
+		TMap<FName, FActivityID> ModifiedPackages;
+
+		enum class ESubobjectState
+		{
+			Created,
+			Removed
+		};
+		
+		/**
+		 * Maps FConcertTransactionEventBase::ExportedObject to the last transaction activity
+		 * that created it (bAllowCreate == true).
+		 */
+		TMap<FSoftObjectPath, TPair<FActivityID, ESubobjectState>> CreatedOrRemovedSubobjects;
+	};
+
+	class FActivityDependencyGraphBuildAlgorithm
+	{
+	public:
+
+		static FActivityDependencyGraph BuildGraph(const FConcertSyncSessionDatabase& SessionDatabase);
+
+	private:
+
+		const FConcertSyncSessionDatabase& SessionDatabase;
+		
+		FActivityDependencyGraph Graph;
+		FPackageTracker PackageTracker;
+
+		FActivityDependencyGraphBuildAlgorithm(const FConcertSyncSessionDatabase& SessionDatabase);
+		
+		bool IsRelevantForDependencies(const FConcertSyncActivity& Activity);
+	
+		void DiscoverPackageDependencies(const FConcertSyncActivity& Activity);
+		void DiscoverPackageDependencies(const FConcertSyncActivity& Activity, const FConcertSyncTransactionEvent& EventData);
+		void DiscoverPackageDependencies(const FConcertSyncActivity& Activity, const FConcertSyncPackageEventMetaData& EventData);
+		
+		void DiscoverAddedPackageDependencies(const FActivityNodeID NodeID, const FConcertSyncActivity& Activity, const FConcertSyncPackageEventMetaData& EventData);
+		void DiscoverSavedPackageDependencies(const FActivityNodeID NodeID, const FConcertSyncActivity& Activity, const FConcertSyncPackageEventMetaData& EventData);
+		void DiscoverRenamedPackageDependencies(const FActivityNodeID NodeID, const FConcertSyncActivity& Activity, const FConcertSyncPackageEventMetaData& EventData);
+		void DiscoverDeletedPackageDependencies(const FActivityNodeID NodeID, const FConcertSyncActivity& Activity, const FConcertSyncPackageEventMetaData& EventData);
+
+		/** Util function for declarative style dependency statement in C++ */
+		enum class EPackageAddDependencyCondition
+		{
+			Always,
+			OnlyLatestActivity
+		};
+		using FPackageActivityItem = TTuple<const FActivityID* /* ActivityID */, EActivityDependencyReason /* Reason */, EDependencyStrength /* DependencyStrength*/, EPackageAddDependencyCondition /* AddCondition */>;
+		using FTrackedPackageActivityArray = TArray<FPackageActivityItem, TInlineAllocator<8>>;
+		void AddDependencies(const FActivityNodeID NodeID, const FTrackedPackageActivityArray& Dependencies);
+	
+		void TrackAffectedPackages(const FConcertSyncActivity& Activity);
+		void TrackAffectedPackages(const FConcertSyncActivity& Activity, const FConcertSyncTransactionEvent& EventData);
+		void TrackAffectedPackages(const FConcertSyncActivity& Activity, const FConcertSyncPackageEventMetaData& EventData);
+
+		FSoftObjectPath MakePathFromExportedObject(const FConcertExportedObject& ExportedObject) const;
+	};
+}
+
+UE::ConcertSyncCore::FActivityDependencyGraph UE::ConcertSyncCore::BuildDependencyGraphFrom(const FConcertSyncSessionDatabase& SessionDatabase)
+{
+	return FActivityDependencyGraphBuildAlgorithm::BuildGraph(SessionDatabase);
+}
+
+namespace UE::ConcertSyncCore
+{
+	/**
+	 * The algorithm works as follows:
+	 *	For each transaction i sorted from earliest to later:
+	 *		1. Add a dependency to every activity that has so far affected a package i depends on.
+	 *		2. Track which packages i added, saved, renamed, removed, or modified.
+	 *		
+	 * Every transaction is processed exactly once, hence this algorithm is O(n).
+	 */
+	FActivityDependencyGraph FActivityDependencyGraphBuildAlgorithm::BuildGraph(const FConcertSyncSessionDatabase& SessionDatabase)
+	{
+		FActivityDependencyGraphBuildAlgorithm Builder(SessionDatabase);
+		SessionDatabase.EnumerateActivities([&Builder](FConcertSyncActivity&& Activity)
+		{
+			if (Builder.IsRelevantForDependencies(Activity))
+			{
+				Builder.DiscoverPackageDependencies(Activity);
+				Builder.TrackAffectedPackages(Activity);
+			}
+
+			return EBreakBehavior::Continue;
+		});
+		
+		return Builder.Graph;
+	}
+
+	FActivityDependencyGraphBuildAlgorithm::FActivityDependencyGraphBuildAlgorithm(const FConcertSyncSessionDatabase& SessionDatabase)
+		: SessionDatabase(SessionDatabase)
+	{}
+
+	bool FActivityDependencyGraphBuildAlgorithm::IsRelevantForDependencies(const FConcertSyncActivity& Activity)
+	{
+		return Activity.EventType == EConcertSyncActivityEventType::Package
+			|| Activity.EventType == EConcertSyncActivityEventType::Transaction;
+	}
+
+	void FActivityDependencyGraphBuildAlgorithm::DiscoverPackageDependencies(const FConcertSyncActivity& Activity, const FConcertSyncTransactionEvent& EventData)
+	{
+		const FActivityID CurrentActivityID = Activity.ActivityId;
+		const FActivityNodeID CurrentNodeID = Graph.AddActivity(CurrentActivityID);
+
+		TSet<FActivityID> CreatedDependencies;
+		for (const FConcertExportedObject& ExportedObject : EventData.Transaction.ExportedObjects)
+		{
+			const TPair<FActivityID, FPackageTracker::ESubobjectState>* LastActivity = PackageTracker.CreatedOrRemovedSubobjects.Find(MakePathFromExportedObject(ExportedObject));
+			if (LastActivity)
+			{
+				const TOptional<FActivityNodeID> NodeID = Graph.FindNodeByActivity(LastActivity->Key);
+				checkf(NodeID, TEXT("If the activity is in CreatedOrRemovedSubobjects, then we should have created a node for it in the graph!"));
+				
+				const EActivityDependencyReason Reason = LastActivity->Value == FPackageTracker::ESubobjectState::Created
+					? EActivityDependencyReason::SubobjectCreation
+					: EActivityDependencyReason::SubobjectRemoval;
+				Graph.AddDependency(CurrentNodeID, FActivityDependencyEdge(*NodeID, Reason, EDependencyStrength::HardDependency));
+
+				CreatedDependencies.Add(LastActivity->Key);
+			}
+		}
+
+		auto AddDependency = [this, CurrentActivityID, CurrentNodeID](const FActivityID DependedOnActivityID, EActivityDependencyReason Reason, EDependencyStrength DependencyStrength)
+		{
+			const TOptional<FActivityNodeID> ActivityNodeID = Graph.FindNodeByActivity(DependedOnActivityID);
+			if (ensureMsgf(ActivityNodeID, TEXT("Investigate algorithm. There should be a node for activity %lld because it supposed to have been processed before activity %lld. We assume that earlier activities have smaller IDs than later activities."), DependedOnActivityID, CurrentActivityID))
+			{
+				Graph.AddDependency(CurrentNodeID, FActivityDependencyEdge(*ActivityNodeID, Reason, DependencyStrength));
+			}
+		};
+		
+		for (const FName ModifiedPackage : EventData.Transaction.ModifiedPackages)
+		{
+			const FActivityID* LastModifyingActivity = PackageTracker.ModifiedPackages.Find(ModifiedPackage);
+			if (LastModifyingActivity)
+			{
+				// If there is already a hard dependency to a previous activity, then we do not "possibly" depend on it: we definitely depend on it (already).
+				if (!CreatedDependencies.Contains(*LastModifyingActivity))
+				{
+					AddDependency(*LastModifyingActivity, EActivityDependencyReason::EditAfterPreviousPackageEdit, EDependencyStrength::PossibleDependency);
+				}
+			}
+			// If nobody modified the package, we depend on the activity that added the package...
+			else if (const FActivityID* LastAddedActivity = PackageTracker.AddedPackages.Find(ModifiedPackage))
+			{
+				AddDependency(*LastAddedActivity, EActivityDependencyReason::PackageCreation, EDependencyStrength::HardDependency);
+			}
+			// ... or "created" it by renaming it
+			else if (const FActivityID* LastSavedActivity = PackageTracker.RenamedPackages.Find(ModifiedPackage))
+			{
+				AddDependency(*LastSavedActivity, EActivityDependencyReason::PackageRename, EDependencyStrength::HardDependency);
+			}
+
+			/* We should consider indirect dependencies, too. Related ticket: UE-148392.
+			 *
+			 * Suppose:
+			 *	1 Create data asset A
+			 *	2 Make actor reference A
+			 *	3 Edit data asset
+			 *	4 Edit actor
+			 *	
+			 * In activity 4, the construction script may depend on the data asset modified in activity 3.
+			 */
+		}
+	}
+
+	void FActivityDependencyGraphBuildAlgorithm::DiscoverPackageDependencies(const FConcertSyncActivity& Activity, const FConcertSyncPackageEventMetaData& EventData)
+	{
+		const FActivityNodeID NodeID = Graph.AddActivity(Activity.ActivityId);
+		switch (EventData.PackageInfo.PackageUpdateType)
+		{
+		case EConcertPackageUpdateType::Added:
+			DiscoverAddedPackageDependencies(NodeID, Activity, EventData);
+			break;
+		case EConcertPackageUpdateType::Saved:
+			DiscoverSavedPackageDependencies(NodeID, Activity, EventData);
+			break;
+		case EConcertPackageUpdateType::Renamed:
+			DiscoverRenamedPackageDependencies(NodeID, Activity, EventData);
+			break;
+		case EConcertPackageUpdateType::Deleted:
+			DiscoverDeletedPackageDependencies(NodeID, Activity, EventData);
+			break;
+			
+		case EConcertPackageUpdateType::Dummy: 
+		default:
+			checkNoEntry();
+		}
+	}
+
+	void FActivityDependencyGraphBuildAlgorithm::DiscoverAddedPackageDependencies(const FActivityNodeID NodeID, const FConcertSyncActivity& Activity,const FConcertSyncPackageEventMetaData& EventData)
+	{
+		const FName& Package = EventData.PackageInfo.PackageName;
+		const FActivityID* LastPackageRemoval = PackageTracker.RemovedPackages.Find(Package);
+		const FActivityID* LastPackageRename = PackageTracker.RenamedPackages.Find(Package);
+		
+		AddDependencies(
+			NodeID,
+			{
+				{ LastPackageRemoval, EActivityDependencyReason::PackageRemoval, EDependencyStrength::HardDependency, EPackageAddDependencyCondition::OnlyLatestActivity },
+				{ LastPackageRename,  EActivityDependencyReason::PackageRename, EDependencyStrength::HardDependency, EPackageAddDependencyCondition::OnlyLatestActivity }
+			});
+	}
+
+	void FActivityDependencyGraphBuildAlgorithm::DiscoverSavedPackageDependencies(const FActivityNodeID NodeID, const FConcertSyncActivity& Activity, const FConcertSyncPackageEventMetaData& EventData)
+	{
+		const FName& Package = EventData.PackageInfo.PackageName;
+		const FActivityID* LastPackageAddition = PackageTracker.AddedPackages.Find(Package);
+		const FActivityID* LastPackageRename = PackageTracker.RenamedPackages.Find(Package);
+
+		const FTrackedPackageActivityArray Dependencies {
+			{ LastPackageAddition, EActivityDependencyReason::PackageCreation, EDependencyStrength::HardDependency, EPackageAddDependencyCondition::OnlyLatestActivity },
+			{ LastPackageRename, EActivityDependencyReason::PackageRename, EDependencyStrength::HardDependency, EPackageAddDependencyCondition::OnlyLatestActivity }
+		};
+		
+		AddDependencies(NodeID, Dependencies);
+	}
+
+	void FActivityDependencyGraphBuildAlgorithm::DiscoverRenamedPackageDependencies(const FActivityNodeID NodeID, const FConcertSyncActivity& Activity, const FConcertSyncPackageEventMetaData& EventData)
+	{
+		const FName& NewPackageName = EventData.PackageInfo.NewPackageName;
+		const FName& OldPackageName = EventData.PackageInfo.PackageName;
+		
+		const FActivityID* NewNameAdded = PackageTracker.AddedPackages.Find(NewPackageName);
+		const FActivityID* NewPackageRenamed = PackageTracker.RenamedPackages.Find(NewPackageName);
+		const FActivityID* SavedNewPackage = PackageTracker.SavedPackages.Find(NewPackageName);
+		FTrackedPackageActivityArray NewPackageDependencies {
+			{ NewNameAdded, EActivityDependencyReason::PackageCreation, EDependencyStrength::HardDependency, EPackageAddDependencyCondition::OnlyLatestActivity },
+			{ NewPackageRenamed, EActivityDependencyReason::PackageRename, EDependencyStrength::HardDependency, EPackageAddDependencyCondition::OnlyLatestActivity },
+		};
+		// When renaming, Concert does not generate a EConcertPackageUpdateType::Added type.
+		// Instead the renamed to package is added via an EConcertPackageUpdateType::Saved activity which occurs just before
+		// the corresponding EConcertPackageUpdateType::Rename activity.
+		if (!NewNameAdded && !NewPackageRenamed)
+		{
+			NewPackageDependencies.Add({ SavedNewPackage, EActivityDependencyReason::PackageCreation, EDependencyStrength::PossibleDependency, EPackageAddDependencyCondition::Always });
+		}
+		
+		const FActivityID* OldPackageAdded = PackageTracker.AddedPackages.Find(OldPackageName);
+		const FActivityID* OldPackageRenamed = PackageTracker.RenamedPackages.Find(OldPackageName);
+		const FActivityID* SavedOldPackage = PackageTracker.SavedPackages.Find(OldPackageName);
+		FTrackedPackageActivityArray OldPackageDependencies {
+			{ OldPackageAdded, EActivityDependencyReason::PackageCreation, EDependencyStrength::HardDependency, EPackageAddDependencyCondition::OnlyLatestActivity },
+			{ OldPackageRenamed, EActivityDependencyReason::PackageRename, EDependencyStrength::HardDependency, EPackageAddDependencyCondition::OnlyLatestActivity },
+		};
+		// See above
+		if (!OldPackageAdded && !OldPackageRenamed)
+		{
+			OldPackageDependencies.Add({ SavedOldPackage, EActivityDependencyReason::PackageCreation, EDependencyStrength::PossibleDependency, EPackageAddDependencyCondition::Always });
+		}
+		
+		AddDependencies(NodeID, NewPackageDependencies);
+		AddDependencies(NodeID, OldPackageDependencies);
+	}
+
+	void FActivityDependencyGraphBuildAlgorithm::DiscoverDeletedPackageDependencies(const FActivityNodeID NodeID, const FConcertSyncActivity& Activity, const FConcertSyncPackageEventMetaData& EventData)
+	{
+		const FName& Package = EventData.PackageInfo.PackageName;
+		const FActivityID* LastPackageAddition = PackageTracker.AddedPackages.Find(Package);
+		const FActivityID* LastPackageRename = PackageTracker.RenamedPackages.Find(Package);
+		
+		AddDependencies(
+			NodeID,
+			{
+				{ LastPackageAddition, EActivityDependencyReason::PackageCreation, EDependencyStrength::HardDependency, EPackageAddDependencyCondition::OnlyLatestActivity},
+				{ LastPackageRename, EActivityDependencyReason::PackageRename, EDependencyStrength::HardDependency, EPackageAddDependencyCondition::OnlyLatestActivity }
+			});
+	}
+	
+	void FActivityDependencyGraphBuildAlgorithm::AddDependencies(const FActivityNodeID NodeID, const FTrackedPackageActivityArray& Dependencies)
+	{
+		auto AddDependency = [this, NodeID](const FPackageActivityItem& Dependency)
+		{
+			const FActivityID* ActivityID = Dependency.Get<0>();
+			const TOptional<FActivityNodeID> TargetNodeID = Graph.FindNodeByActivity(*ActivityID);
+			if (!ensureMsgf(TargetNodeID, TEXT("Investigate algorithm. There should be a node for activity %lld because it supposed to have been processed before activity %lld. We assume that earlier activities have smaller IDs than later activities."), *ActivityID, NodeID.ID))
+			{
+				return;
+			}
+
+			const EActivityDependencyReason Reason = Dependency.Get<1>();
+			Graph.AddDependency(NodeID, FActivityDependencyEdge(*TargetNodeID, Reason, Dependency.Get<2>()));
+		};
+		
+		const FPackageActivityItem* LatestDependency = nullptr;
+		for (const FPackageActivityItem& Dependency : Dependencies)
+		{
+			const FActivityID* DependencyActivityID = Dependency.Get<0>();
+			if (!DependencyActivityID)
+			{
+				continue;
+			}
+			
+			if (Dependency.Get<3>() == EPackageAddDependencyCondition::Always)
+			{
+				AddDependency(Dependency);
+			}
+			else 
+			{
+				check(Dependency.Get<3>() == EPackageAddDependencyCondition::OnlyLatestActivity);
+				const bool bCurrentActivityIsNewer = !LatestDependency || *LatestDependency->Get<0>() < *DependencyActivityID;
+				if (bCurrentActivityIsNewer)
+				{
+					LatestDependency = &Dependency;
+				}
+			}
+		}
+
+		if (LatestDependency)
+		{
+			AddDependency(*LatestDependency);
+		}
+	}
+
+	void FActivityDependencyGraphBuildAlgorithm::TrackAffectedPackages(const FConcertSyncActivity& Activity, const FConcertSyncTransactionEvent& EventData)
+	{
+		for (FName ModifiedPackage : EventData.Transaction.ModifiedPackages)
+		{
+			PackageTracker.ModifiedPackages.Add(ModifiedPackage, Activity.ActivityId);
+		}
+		
+		for (const FConcertExportedObject& ExportedObject : EventData.Transaction.ExportedObjects)
+		{
+			if (ExportedObject.ObjectData.bAllowCreate && ensure(!ExportedObject.ObjectData.bIsPendingKill))
+			{
+				PackageTracker.CreatedOrRemovedSubobjects.Add(MakePathFromExportedObject(ExportedObject), { Activity.ActivityId, FPackageTracker::ESubobjectState::Created });
+			}
+			else if (ExportedObject.ObjectData.bIsPendingKill && ensure(!ExportedObject.ObjectData.bAllowCreate))
+			{
+				PackageTracker.CreatedOrRemovedSubobjects.Add(MakePathFromExportedObject(ExportedObject), { Activity.ActivityId, FPackageTracker::ESubobjectState::Removed });
+			}
+		}
+	}
+
+	void FActivityDependencyGraphBuildAlgorithm::TrackAffectedPackages(const FConcertSyncActivity& Activity, const FConcertSyncPackageEventMetaData& EventData)
+	{
+		const FName PackageName = EventData.PackageInfo.PackageName;
+		const FActivityID ActivityID = Activity.ActivityId;
+		switch (EventData.PackageInfo.PackageUpdateType)
+		{
+		case EConcertPackageUpdateType::Added:
+			PackageTracker.AddedPackages.Add(PackageName, ActivityID);
+			break;
+			
+		case EConcertPackageUpdateType::Saved:
+			PackageTracker.SavedPackages.Add(PackageName, ActivityID);
+			break;
+			
+		case EConcertPackageUpdateType::Renamed:
+			PackageTracker.RenamedPackages.Add(EventData.PackageInfo.NewPackageName, ActivityID);
+			break;
+			
+		case EConcertPackageUpdateType::Deleted:
+			PackageTracker.RemovedPackages.Add(PackageName, ActivityID);
+			break;
+			
+		case EConcertPackageUpdateType::Dummy:
+			// No idea what dependencies this has ... let's assume none and hopefully not regret it later
+			break;
+			
+		default:
+			checkNoEntry();
+		}
+	}
+
+	FSoftObjectPath FActivityDependencyGraphBuildAlgorithm::MakePathFromExportedObject(const FConcertExportedObject& ExportedObject) const
+	{
+		return FSoftObjectPath(FString::Printf(TEXT("%s.%s"), *ExportedObject.ObjectId.ObjectOuterPathName.ToString(), *ExportedObject.ObjectId.ObjectName.ToString()));
+	}
+
+/** Checks whether Activity is a transaction or package activity, gets the event meta data, and forwards it to an overload of FuncName. */
+#define FORWARD_ACTIVITY(Activity, FuncName) FConcertSyncTransactionEvent TransactionEvent; \
+	FConcertSyncPackageEventMetaData PackageEventMetaData; \
+	switch (Activity.EventType) \
+	{ \
+	case EConcertSyncActivityEventType::Transaction: \
+		if (SessionDatabase.GetTransactionEvent(Activity.EventId, TransactionEvent)) \
+		{ \
+			FuncName(Activity, TransactionEvent); \
+		} \
+		else \
+		{ \
+			checkNoEntry(); \
+		} \
+		break; \
+	case EConcertSyncActivityEventType::Package: \
+		if (SessionDatabase.GetPackageEventMetaData(Activity.EventId, PackageEventMetaData.PackageRevision, PackageEventMetaData.PackageInfo)) \
+		{ \
+			FuncName(Activity, PackageEventMetaData); \
+		} \
+		else \
+		{ \
+			checkNoEntry(); \
+		} \
+		break; \
+	case EConcertSyncActivityEventType::None: \
+	case EConcertSyncActivityEventType::Connection: \
+	case EConcertSyncActivityEventType::Lock: \
+	default: \
+		checkNoEntry(); \
+	} \
+
+	void FActivityDependencyGraphBuildAlgorithm::DiscoverPackageDependencies(const FConcertSyncActivity& Activity)
+	{
+		FORWARD_ACTIVITY(Activity, DiscoverPackageDependencies)
+	}
+
+	void FActivityDependencyGraphBuildAlgorithm::TrackAffectedPackages(const FConcertSyncActivity& Activity)
+	{
+		FORWARD_ACTIVITY(Activity, TrackAffectedPackages)
+	}
+}
\ No newline at end of file
diff --git a/Engine/Plugins/Developer/Concert/ConcertSync/ConcertSyncCore/Source/ConcertSyncCore/Public/HistoryEdition/ActivityDependencyEdge.h b/Engine/Plugins/Developer/Concert/ConcertSync/ConcertSyncCore/Source/ConcertSyncCore/Public/HistoryEdition/ActivityDependencyEdge.h
new file mode 100644
--- /dev/null
+++ b/Engine/Plugins/Developer/Concert/ConcertSync/ConcertSyncCore/Source/ConcertSyncCore/Public/HistoryEdition/ActivityDependencyEdge.h
@@ -0,0 +1,124 @@
+// Copyright Epic Games, Inc. All Rights Reserved.
+
+#pragma once
+
+#include "CoreMinimal.h"
+#include "ActivityGraphIDs.h"
+
+namespace UE::ConcertSyncCore
+{
+	class FActivityNode;
+	class FActivityDependencyGraph;
+	
+	/** Describes why there is a dependency */
+	enum class EActivityDependencyReason
+	{
+		/**
+		 * The depended-on activity created the package that this activity operates on.
+		 * 
+		 * Example: 1: Create level > 2: Add actor
+		 * 2 cannot take place without 1.
+		 */
+		PackageCreation,
+
+		/**
+		 * The depended-on activity deleted a package with the same path name of the package this activity created.
+		 * 
+		 * Example: 1: Delete level Game/Foo 2: Create new level Game/Foo.
+		 * 2 cannot take place without 1.
+		 */
+		PackageRemoval,
+
+		/**
+		 * The depended-on activity renamed a package that is used by this activity.
+		 * 
+		 * Example: 1: Create level 'Foo' > 2: Rename 'Foo' to 'Bar' > 3: Edit actor
+		 * 3 depends on 2.
+		 * 
+		 * Example: 1: Create level 'Foo' > 2: Rename 'Foo' to 'Bar' > 3: Create level 'Foo'
+		 * 3 depends on 1.
+		 */
+		PackageRename,
+
+		/**
+		 * The depended-on activity changes the state of a package that is depended on by a package this activity modifies.
+		 * 
+		 * Example: 1: Edit actor to reference data asset 2: Edit data asset 3: Edit actor, triggering the construction script.
+		 * 3 executes the construction script which may execute differently depending on the data changed in 2.
+		 */
+		EditAfterPreviousPackageEdit,
+
+		/**
+		 * The depended-on activity created a subobject that this activity operates on.
+		 * 
+		 * Example: 1: Add actor > 2. Edit actor
+		 * 2 cannot take place without 1.
+		 */
+		SubobjectCreation,
+
+		/**
+		 * The depended-on activity removed a subobject that this activity requires to be dead.
+		 * 
+		 * Example: 1: Remove actor > 2. Create actor with same name
+		 * 2 cannot take place without 1.
+		 */
+		SubobjectRemoval
+	};
+	
+	FString LexToString(EActivityDependencyReason Reason);
+
+	enum class EDependencyStrength
+	{
+		/**
+		 * An activity A has a hard dependency to activity B when A cannot exist in an activity history without B.
+		 * Example: You create a level and add an actor to it. You cannot create the actor without the level.
+		 */
+		HardDependency,
+		
+		/**
+		 * An activity A has a possible dependency to activity B when A may exist with B but we cannot rule out
+		 * that B does not affect A.
+		 * Example: You modify an actor twice each time triggering the construction script
+		 */
+		PossibleDependency,
+	};
+
+
+	/** Describes a dependency from one activity to an earlier activity. */
+	class CONCERTSYNCCORE_API FActivityDependencyEdge
+	{
+	public:
+		
+		FActivityDependencyEdge(FActivityNodeID DependedOnActivityNode, EActivityDependencyReason DependencyReason, EDependencyStrength DependencyStrength)
+			: Parent(DependedOnActivityNode)
+			, DependencyReason(DependencyReason)
+			, DependencyStrength(DependencyStrength)
+		{}
+		
+		FActivityNodeID GetDependedOnNodeID() const { return Parent; }
+		EActivityDependencyReason GetDependencyReason() const { return DependencyReason; }
+		EDependencyStrength GetDependencyStrength() const { return DependencyStrength; }
+
+		friend bool operator==(const FActivityDependencyEdge& Left, const FActivityDependencyEdge& Right)
+		{
+			return Left.Parent == Right.Parent
+				&& Left.DependencyReason == Right.DependencyReason
+				&& Left.DependencyStrength == Right.DependencyStrength;
+		}
+		friend bool operator!=(const FActivityDependencyEdge& Left, const FActivityDependencyEdge& Right)
+		{
+			return !(Left == Right);
+		}
+
+	private:
+		
+		/** The node that an activity depends on */
+		const FActivityNodeID Parent;
+
+		/** Why this dependency exists */
+		const EActivityDependencyReason DependencyReason;
+
+		/** How much the activity is depended on */
+		const EDependencyStrength DependencyStrength;
+	};
+}
\ No newline at end of file
diff --git a/Engine/Plugins/Developer/Concert/ConcertSync/ConcertSyncCore/Source/ConcertSyncCore/Public/HistoryEdition/ActivityDependencyGraph.h b/Engine/Plugins/Developer/Concert/ConcertSync/ConcertSyncCore/Source/ConcertSyncCore/Public/HistoryEdition/ActivityDependencyGraph.h
new file mode 100644
--- /dev/null
+++ b/Engine/Plugins/Developer/Concert/ConcertSync/ConcertSyncCore/Source/ConcertSyncCore/Public/HistoryEdition/ActivityDependencyGraph.h
@@ -0,0 +1,82 @@
+﻿// Copyright Epic Games, Inc. All Rights Reserved.
+
+#pragma once
+
+#include "CoreMinimal.h"
+#include "ActivityDependencyEdge.h"
+#include "ActivityGraphIDs.h"
+#include "ActivityNode.h"
+#include "Misc/Optional.h"
+#include "Templates/NonNullPointer.h"
+
+namespace UE::ConcertSyncCore
+{
+	/**
+	 * The graph models activities as nodes and dependencies to activities as edges.
+	 * Activities can only affect each other from left-to-right: an earlier activity affects a later activity.
+	 * 
+	 * We say an activity A has a hard dependency to activity B when A cannot exist in an activity history without B.
+	 * We say an activity A has a possible dependency to activity B when A may exist with B but we cannot rule out
+	 * that B does not affect A. Example: You modify an actor twice each time triggering the construction script.
+	 * See EDependencyType.
+	 *
+	 * A node corresponds to an activity. A root node has depends on no other activity.
+	 * An edge from a A to B means that B depends on A.
+	 * 
+	 * Dependencies do not necessarily need to be transitive: 1. Remove package Foo 2. Create package Foo 3. Edit Foo
+	 * The existence of 3 can requires that either
+	 *  - 1 and 2 both exist
+	 *  - 1 and 2 are deleted (among other details, e.g. that the underlying assets must have equal classes)
+	 *      
+	 * Example:
+	 *	1: Create level named Foo
+	 *	2: Create level named Bar
+	 *	3: Add actor named 'Cube' to level Foo
+	 *	4. Add actor named 'Sphere" to level Foo
+	 *	6: Edit 'Cube'.
+	 * Resulting graph:
+	 *			1		2
+	 *		   / \
+	 *		  3	  4
+	 *		 /
+	 *		5
+	 *
+	 * Example:
+	 *  1. Create data asset A
+	 *  2. Edit and save A (remember: transactions do not happen for all .uasset only for .umaps)
+	 *  3. Create actor in level
+	 *  4. Edit actor to reference data assets A
+	 *  5. Edit actor whose construction script reads data from A
+	 * Resulting Graph:     
+	 *     1   3
+	 *    / \ / 
+	 *   2   4
+	 *       I
+	 *       5
+	 */
+	class CONCERTSYNCCORE_API FActivityDependencyGraph
+	{
+	public:
+		
+		/** Adds a node. Invalidates the results of GetNodeById. */
+		FActivityNodeID AddActivity(int64 ActivityIndex);
+		
+		/** Adds a dependency from an existing node to another existing node */
+		bool AddDependency(FActivityNodeID From, FActivityDependencyEdge To);
+		
+		/** Tries to find a corresponding node for an activity ID. */
+		TOptional<FActivityNodeID> FindNodeByActivity(int64 ActivityID) const;
+
+		/** Gets the node identified by the given ID, if it is valid. Only valid as long as no new node is added.*/
+		const FActivityNode& GetNodeById(FActivityNodeID ID) const;
+
+		void ForEachNode(TFunctionRef<void(const FActivityNode&)> ConsumerFunc) const;
+		
+	private:
+
+		TArray<FActivityNode> Nodes;
+		
+		FActivityNode& GetNodeByIdInternal(FActivityNodeID ID);
+	};
+}
+
diff --git a/Engine/Plugins/Developer/Concert/ConcertSync/ConcertSyncCore/Source/ConcertSyncCore/Public/HistoryEdition/ActivityGraphIDs.h b/Engine/Plugins/Developer/Concert/ConcertSync/ConcertSyncCore/Source/ConcertSyncCore/Public/HistoryEdition/ActivityGraphIDs.h
new file mode 100644
--- /dev/null
+++ b/Engine/Plugins/Developer/Concert/ConcertSync/ConcertSyncCore/Source/ConcertSyncCore/Public/HistoryEdition/ActivityGraphIDs.h
@@ -0,0 +1,47 @@
+﻿// Copyright Epic Games, Inc. All Rights Reserved.
+
+#pragma once
+
+#include "CoreMinimal.h"
+#include "Templates/TypeHash.h"
+
+namespace UE::ConcertSyncCore
+{
+	class FActivityDependecyGraph;
+	/** Enforces type-safe usage of IDs by avoiding implicit conversions. */
+	struct CONCERTSYNCCORE_API FActivityNodeID
+	{
+		size_t ID;
+
+		FActivityNodeID() = default;
+		
+		template<typename T>
+		explicit FActivityNodeID(T&& ID)
+			: ID(Forward<T>(ID))
+		{}
+
+		FActivityNodeID(const FActivityNodeID& Other)
+			: ID(Other.ID)
+		{}
+
+		bool HasAnyDependency(const FActivityDependecyGraph&) const;
+
+		/** Conversion back to to size_t is not a common programmer mistake */
+		explicit operator size_t() const { return ID; }
+
+		friend bool operator==(const FActivityNodeID& Left, const FActivityNodeID& Right)
+		{
+			return Left.ID == Right.ID;
+		}
+
+		friend bool operator!=(const FActivityNodeID& Left, const FActivityNodeID& Right)
+		{
+			return !(Left == Right);
+		}
+
+		FORCEINLINE friend uint32 GetTypeHash(const FActivityNodeID& NodeID)
+		{
+			return ::GetTypeHash(NodeID.ID);
+		}
+	};
+}
\ No newline at end of file
diff --git a/Engine/Plugins/Developer/Concert/ConcertSync/ConcertSyncCore/Source/ConcertSyncCore/Public/HistoryEdition/ActivityNode.h b/Engine/Plugins/Developer/Concert/ConcertSync/ConcertSyncCore/Source/ConcertSyncCore/Public/HistoryEdition/ActivityNode.h
new file mode 100644
--- /dev/null
+++ b/Engine/Plugins/Developer/Concert/ConcertSync/ConcertSyncCore/Source/ConcertSyncCore/Public/HistoryEdition/ActivityNode.h
@@ -0,0 +1,57 @@
+// Copyright Epic Games, Inc. All Rights Reserved.
+
+#pragma once
+
+#include "CoreMinimal.h"
+#include "ActivityDependencyEdge.h"
+#include "ActivityGraphIDs.h"
+#include "Misc/Optional.h"
+#include "Templates/NonNullPointer.h"
+
+namespace UE::ConcertSyncCore
+{
+	class FActivityDependencyGraph;
+
+	/** A node corresponds to an activity and can depend on preceding activity*/
+	class CONCERTSYNCCORE_API FActivityNode
+	{
+		friend FActivityDependencyGraph;
+	public:
+		
+		FActivityNode(int64 ActivityId, FActivityNodeID NodeIndex)
+			: ActivityId(ActivityId)
+			, NodeIndex(NodeIndex)
+		{}
+
+		bool HasAnyDependency() const { return Dependencies.Num() > 0; }
+		bool DependsOnNode(FActivityNodeID NodeId, TOptional<EActivityDependencyReason> WithReason = {}, TOptional<EDependencyStrength> WithStrength = {}) const; 
+		bool DependsOnActivity(int64 InActivityId, const FActivityDependencyGraph& Graph, TOptional<EActivityDependencyReason> WithReason = {}, TOptional<EDependencyStrength> WithStrength = {}) const;
+		
+		bool AffectsAnyActivity() const { return GetAffectedChildren().Num() > 0; }
+		bool AffectsNode(FActivityNodeID NodeID) const { return AffectedChildren.Contains(NodeID); }
+		bool AffectsActivity(int64 InActivityId, const FActivityDependencyGraph& Graph) const;
+		
+		int64 GetActivityId() const { return ActivityId; }
+		FActivityNodeID GetNodeIndex() const { return NodeIndex; }
+		const TArray<FActivityDependencyEdge>& GetDependencies() const { return Dependencies; }
+		const TArray<FActivityNodeID>& GetAffectedChildren() const { return AffectedChildren; }
+	
+	private:
+	
+		/** The activity this node corresponds to */
+		const int64 ActivityId;
+
+		/** Index in parent tree's node array (to avoid bad memory footprint) */
+		const FActivityNodeID NodeIndex;
+
+		/** This node's parents. Unset implies this is the root node. */
+		TArray<FActivityDependencyEdge> Dependencies;
+
+		/**
+		 * This node's children. Children have dependencies to this node.
+		 * Activities can only depend on activities that precede them.
+		 * Since activity IDs are historic, every child's activity ID > ActivityId.
+		 */
+		TArray<FActivityNodeID> AffectedChildren;
+	};
+}
\ No newline at end of file
diff --git a/Engine/Plugins/Developer/Concert/ConcertSync/ConcertSyncCore/Source/ConcertSyncCore/Public/HistoryEdition/DebugDependencyGraph.h b/Engine/Plugins/Developer/Concert/ConcertSync/ConcertSyncCore/Source/ConcertSyncCore/Public/HistoryEdition/DebugDependencyGraph.h
new file mode 100644
--- /dev/null
+++ b/Engine/Plugins/Developer/Concert/ConcertSync/ConcertSyncCore/Source/ConcertSyncCore/Public/HistoryEdition/DebugDependencyGraph.h
@@ -0,0 +1,64 @@
+﻿// Copyright Epic Games, Inc. All Rights Reserved.
+
+#pragma once
+
+#include "CoreMinimal.h"
+#include "ActivityDependencyGraph.h"
+#include "ActivityGraphIDs.h"
+
+class FConcertSyncSessionDatabase;
+
+namespace UE::ConcertSyncCore::Graphviz
+{
+	/**
+	 * Graphviz syntax:
+	 *
+	 *	strict digraph {
+	 * 		"start node name" [shape="ellipse" style="filled" fillcolor="#1f77b4"]
+	 *		"end node name" [shape="polygon" style="filled" fillcolor="#ff7f0e"]
+	 *		"start node name" -> "end node name" [fillcolor="#a6cee3" color="#1f78b4" label="Edge label"]
+	 *	}	
+	 */
+
+	using FGraphStringBuilder = TStringBuilder<512>;
+	using FStringifyNode = TFunctionRef<void(FGraphStringBuilder& WriteTo, FActivityNodeID ToStringify)>;
+	using FStringifyEdge = TFunctionRef<void(FGraphStringBuilder& WriteTo, FActivityNodeID EdgeStart, const FActivityDependencyEdge& ToStringify)>;
+
+	/** Name of node without quotes, e.g. TEXT("start node name") (see above) */
+	using FMakeNodeTitle = FStringifyNode;
+	/** Style of the node which is printed right after the name, e.g. TEXT("[shape="ellipse" style="filled" fillcolor="#1f77b4"]") (see above) */
+	using FGetNodeStyle = FStringifyNode;
+	/** Style of the edge which is printed right after an edge, e.g. TEXT("[fillcolor="#a6cee3" color="#1f78b4" label="Edge label"]") (see above) */
+	using FGetEdgeStyle = FStringifyEdge;
+
+	enum class ENodeTitleFlags
+	{
+		None = 0x00,
+		/** Use FConcertSyncActivitySummary::ToDisplayText */
+		Summary = 0x01,
+		/** Append activity ID */
+		ActivityID = 0x02,
+		
+		Default = Summary | ActivityID
+	};
+	ENUM_CLASS_FLAGS(ENodeTitleFlags)
+	
+	/** Prints the title of the activity */
+	CONCERTSYNCCORE_API void MakeNodeTitle(FGraphStringBuilder& WriteTo, FActivityNodeID ToStringify, const FActivityDependencyGraph& Graph, const FConcertSyncSessionDatabase& SessionDatabase, ENodeTitleFlags NodeTitleFlags = ENodeTitleFlags::Default);
+	/** Changes colour depending on activity type*/
+	CONCERTSYNCCORE_API void GetNodeStyle(FGraphStringBuilder& WriteTo, FActivityNodeID ToStringify, const FActivityDependencyGraph& Graph, const FConcertSyncSessionDatabase& SessionDatabase);
+	/** Prints dependency reason and changes the arrow icon depending on dependency strength (half arrow for PossibleDependency, full arrow for HardDependency). */
+	CONCERTSYNCCORE_API void GetEdgeStyle(FGraphStringBuilder& WriteTo, const FActivityDependencyEdge& ToStringify);
+
+	
+	/** Calls ExportToGraphviz with the default formatting functions defined above. */
+	CONCERTSYNCCORE_API FString ExportToGraphviz(const FActivityDependencyGraph& Graph, const FConcertSyncSessionDatabase& SessionDatabase, ENodeTitleFlags NodeTitleFlags = ENodeTitleFlags::Default);
+	/**
+	 * Exports the graph in a syntax that Graphviz (see:http://www.graphviz.org/) understands.
+	 * 
+	 * Several online editors exist and can be found by searching for "graphviz online", e.g.
+	 *	- https://dreampuf.github.io/GraphvizOnline/
+	 *	- http://magjac.com/graphviz-visual-editor/
+	 */
+	CONCERTSYNCCORE_API FString ExportToGraphviz(const FActivityDependencyGraph& Graph, FMakeNodeTitle MakeNodeTitleFunc, FGetNodeStyle GetNodeStyleFunc, FGetEdgeStyle GetEdgeStyleFunc);
+}
diff --git a/Engine/Plugins/Developer/Concert/ConcertSync/ConcertSyncCore/Source/ConcertSyncCore/Public/HistoryEdition/DependencyGraphBuilder.h b/Engine/Plugins/Developer/Concert/ConcertSync/ConcertSyncCore/Source/ConcertSyncCore/Public/HistoryEdition/DependencyGraphBuilder.h
new file mode 100644
--- /dev/null
+++ b/Engine/Plugins/Developer/Concert/ConcertSync/ConcertSyncCore/Source/ConcertSyncCore/Public/HistoryEdition/DependencyGraphBuilder.h
@@ -0,0 +1,16 @@
+﻿// Copyright Epic Games, Inc. All Rights Reserved.
+
+#pragma once
+
+#include "CoreMinimal.h"
+#include "ActivityDependencyGraph.h"
+
+class FConcertSyncSessionDatabase;
+
+namespace UE::ConcertSyncCore
+{
+	/**
+	 * Builds a dependency graph from the given session database.
+	 */
+	CONCERTSYNCCORE_API FActivityDependencyGraph BuildDependencyGraphFrom(const FConcertSyncSessionDatabase& SessionDatabase);
+}
diff --git a/Engine/Plugins/Developer/Concert/ConcertSync/ConcertSyncTest/Source/ConcertSyncTest/Private/HistoryEdition/DependencyGraphBuilderTests.cpp b/Engine/Plugins/Developer/Concert/ConcertSync/ConcertSyncTest/Source/ConcertSyncTest/Private/HistoryEdition/DependencyGraphBuilderTests.cpp
new file mode 100644
--- /dev/null
+++ b/Engine/Plugins/Developer/Concert/ConcertSync/ConcertSyncTest/Source/ConcertSyncTest/Private/HistoryEdition/DependencyGraphBuilderTests.cpp
@@ -0,0 +1,298 @@
+// Copyright Epic Games, Inc. All Rights Reserved.
+
+#include "ConcertLogGlobal.h"
+
+#include "ConcertSyncSessionDatabase.h"
+#include "HistoryEdition/DebugDependencyGraph.h"
+#include "HistoryEdition/DependencyGraphBuilder.h"
+#include "RenameEditAndDeleteMapsFlow.h"
+#include "ScopedSessionDatabase.h"
+
+#include "Misc/AutomationTest.h"
+
+namespace UE::ConcertSyncTests::RenameEditAndDeleteMapsFlowTest
+{
+	/** Validates that the graph reflects the expected dependencies. */
+	bool ValidateExpectedDependencies(FAutomationTestBase& Test, const TTestActivityArray<int64>& ActivityMappings, const ConcertSyncCore::FActivityDependencyGraph& Graph);
+	/** Validates that each activity has a node in the dependency graph. */
+	TTestActivityArray<ConcertSyncCore::FActivityNodeID> ValidateEachActivityHasNode(FAutomationTestBase& Test, const TTestActivityArray<int64>& ActivityMappings, const ConcertSyncCore::FActivityDependencyGraph& Graph);
+	
+	/**
+	 * Builds the dependency graph from a typical sequence of events.
+	 *
+	 * Sequence of user actions:
+	 *	1 Create map Foo
+	 *	2 Add actor A
+	 *	3 Rename actor A
+	 *	4 Edit actor A
+	 *	5 Rename map to Bar
+	 *	6 Edit actor A
+	 *	7 Delete map Bar
+	 *	8 Create map Bar
+	 *
+	 *	The dependency graph should look like this:
+	 *	2 -> 1 (PackageCreation)
+	 *	3 -> 2 (EditPossiblyDependsOnPackage)
+	 *	4 -> 3 (EditPossiblyDependsOnPackage)
+	 *	5 -> 1 (PackageCreation)
+	 *	6 -> 5 (PackageRename)
+	 *	7 -> 5 (PackageRename)
+	 *	8 -> 5 (PackageRemoval)
+	 */
+	IMPLEMENT_SIMPLE_AUTOMATION_TEST(FRenameEditAndDeleteMapsFlowTest, "Concert.History.BuildGraph.RenameEditAndDeleteMapsFlow", EAutomationTestFlags::ApplicationContextMask | EAutomationTestFlags::SmokeFilter)
+	bool FRenameEditAndDeleteMapsFlowTest::RunTest(const FString& Parameters)
+	{
+		FScopedSessionDatabase SessionDatabase(*this);
+		const TTestActivityArray<int64> Activities = CreateActivityHistory(SessionDatabase, SessionDatabase.GetEndpoint());
+		
+		const ConcertSyncCore::FActivityDependencyGraph DependencyGraph = ConcertSyncCore::BuildDependencyGraphFrom(SessionDatabase);
+		UE_LOG(LogConcert, Log, TEXT("%s tested graph in Graphviz format:\n\n%s"), *GetTestFullName(), *ConcertSyncCore::Graphviz::ExportToGraphviz(DependencyGraph, SessionDatabase));
+		return ValidateExpectedDependencies(*this, Activities, DependencyGraph);
+	}
+	
+	bool ValidateExpectedDependencies(FAutomationTestBase& Test, const TTestActivityArray<int64>& ActivityMappings, const ConcertSyncCore::FActivityDependencyGraph& Graph)
+	{
+		using namespace ConcertSyncCore;
+		
+		const TTestActivityArray<FActivityNodeID> ActivityNodes
+			= ValidateEachActivityHasNode(Test, ActivityMappings, Graph);
+
+		// 1 Create map Foo
+		{
+			Test.TestFalse(TEXT("1 Creating new package 'Foo' has no dependencies."), Graph.GetNodeById(ActivityNodes[_1_NewPackageFoo]).HasAnyDependency());
+			
+			Test.TestTrue(TEXT("1 Saving 'Foo' has dependency to creating package 'Foo'."), Graph.GetNodeById(ActivityNodes[_1_SavePackageFoo]).DependsOnActivity(ActivityMappings[_1_NewPackageFoo], Graph, EActivityDependencyReason::PackageCreation, EDependencyStrength::HardDependency));
+			Test.TestEqual(TEXT("1 Saving 'Foo' has exactly 1 dependency"), Graph.GetNodeById(ActivityNodes[_1_SavePackageFoo]).GetDependencies().Num(), 1);
+		}
+
+		// 2 Add actor A
+		{
+			Test.TestTrue(TEXT("2 Adding actor to 'Foo' depends on creating package 'Foo'."), Graph.GetNodeById(ActivityNodes[_2_AddActor]).DependsOnActivity(ActivityMappings[_1_NewPackageFoo], Graph, EActivityDependencyReason::PackageCreation, EDependencyStrength::HardDependency));
+			Test.TestEqual(TEXT("2 Saving 'Foo' has exactly 1 dependency"), Graph.GetNodeById(ActivityNodes[_2_AddActor]).GetDependencies().Num(), 1);
+		}
+
+		// 3 Rename actor A
+		{
+			// It must be a EDependencyStrength::HardDependency because you cannot edit the actor without having created it 
+			Test.TestTrue(TEXT("3 Renaming actor depends on having created the actor."), Graph.GetNodeById(ActivityNodes[_3_RenameActor]).DependsOnActivity(ActivityMappings[_2_AddActor], Graph, EActivityDependencyReason::SubobjectCreation, EDependencyStrength::HardDependency));
+			Test.TestEqual(TEXT("3 Renaming actor has exactly 1 dependency"), Graph.GetNodeById(ActivityNodes[_3_RenameActor]).GetDependencies().Num(), 1);
+		}
+
+		// 4 Edit actor A
+		{
+			// The previous edit might have affected us (e.g. this activity may have executed the construction script)
+			// Note: This should not have a hard dependency on having renamed the actor because a rename is just a property change of ActorLabel.
+			Test.TestTrue(TEXT("4 Editing actor may depend on having edited the actor previously."), Graph.GetNodeById(ActivityNodes[_4_EditActor]).DependsOnActivity(ActivityMappings[_3_RenameActor], Graph, EActivityDependencyReason::EditAfterPreviousPackageEdit, EDependencyStrength::PossibleDependency));
+			// This activity must have a EDependencyStrength::HardDependency to _2_AddActor because the edit cannot happen without having created the actor
+			Test.TestTrue(TEXT("4 Editing actor depends on having created the actor."), Graph.GetNodeById(ActivityNodes[_4_EditActor]).DependsOnActivity(ActivityMappings[_2_AddActor], Graph, EActivityDependencyReason::SubobjectCreation, EDependencyStrength::HardDependency));
+			Test.TestEqual(TEXT("4 Editing actor has exactly 1 dependency"), Graph.GetNodeById(ActivityNodes[_4_EditActor]).GetDependencies().Num(), 2);
+		}
+
+		// 5 Rename map to Bar
+		{
+			Test.TestFalse(TEXT("5 Saving new package 'Bar' has no dependencies."), Graph.GetNodeById(ActivityNodes[_5_SavePackageBar]).HasAnyDependency());
+			
+			Test.TestTrue(TEXT("5 Renaming 'Foo' to 'Bar' has dependency to creating package 'Foo'."), Graph.GetNodeById(ActivityNodes[_5_RenameFooToBar]).DependsOnActivity(ActivityMappings[_1_NewPackageFoo], Graph, EActivityDependencyReason::PackageCreation, EDependencyStrength::HardDependency));
+			Test.TestTrue(TEXT("5 Renaming 'Foo' to 'Bar' has dependency to creating package 'Bar'."), Graph.GetNodeById(ActivityNodes[_5_RenameFooToBar]).DependsOnActivity(ActivityMappings[_5_SavePackageBar], Graph, EActivityDependencyReason::PackageCreation, EDependencyStrength::PossibleDependency));
+			Test.TestEqual(TEXT("5 Saving 'Foo' has exactly 2 dependencies"), Graph.GetNodeById(ActivityNodes[_5_RenameFooToBar]).GetDependencies().Num(), 2);
+		}
+
+		// 6 Edit actor A
+		{
+			Test.TestTrue(TEXT("6 Editing actor in 'Bar' depends on having renamed 'Foo' to 'Bar'."), Graph.GetNodeById(ActivityNodes[_6_EditActor]).DependsOnActivity(ActivityMappings[_5_RenameFooToBar], Graph, EActivityDependencyReason::PackageRename, EDependencyStrength::HardDependency));
+			Test.TestEqual(TEXT("6 Editing actor has exactly 1 dependency"), Graph.GetNodeById(ActivityNodes[_6_EditActor]).GetDependencies().Num(), 1);
+		}
+
+		// 7 Delete map Bar
+		{
+			Test.TestTrue(TEXT("7 Deleting 'Bar' depends on having renamed 'Foo' to 'Bar' previously."), Graph.GetNodeById(ActivityNodes[_7_DeleteBar]).DependsOnActivity(ActivityMappings[_5_RenameFooToBar], Graph, EActivityDependencyReason::PackageRename, EDependencyStrength::HardDependency));
+			Test.TestEqual(TEXT("7 Deleting 'Bar' after rename has exactly 1 dependency"), Graph.GetNodeById(ActivityNodes[_7_DeleteBar]).GetDependencies().Num(), 1);
+		}
+
+		// 8 Create map Bar
+		{
+			Test.TestTrue(TEXT("8 Re-creating 'Bar' depends on having deleted 'Bar' previously."), Graph.GetNodeById(ActivityNodes[_8_NewPackageFoo]).DependsOnActivity(ActivityMappings[_7_DeleteBar], Graph, EActivityDependencyReason::PackageRemoval, EDependencyStrength::HardDependency));
+			Test.TestEqual(TEXT("8 Re-creating 'Bar' 1 dependency"), Graph.GetNodeById(ActivityNodes[_8_NewPackageFoo]).GetDependencies().Num(), 1);
+			
+			Test.TestTrue(TEXT("8 Saving 'Bar' depends on re-created 'Bar'."), Graph.GetNodeById(ActivityNodes[_8_SavePackageFoo]).DependsOnActivity(ActivityMappings[_8_NewPackageFoo], Graph, EActivityDependencyReason::PackageCreation, EDependencyStrength::HardDependency));
+			Test.TestEqual(TEXT("8 Saving 'Bar' after re-creation has exactly 1 dependency"), Graph.GetNodeById(ActivityNodes[_8_SavePackageFoo]).GetDependencies().Num(), 1);
+		}
+		
+		return true;
+	}
+
+	TTestActivityArray<ConcertSyncCore::FActivityNodeID> ValidateEachActivityHasNode(FAutomationTestBase& Test, const TTestActivityArray<int64>& ActivityMappings, const ConcertSyncCore::FActivityDependencyGraph& Graph)
+	{
+		using namespace ConcertSyncCore;
+		
+		TTestActivityArray<FActivityNodeID> ActivityNodes;
+		ActivityNodes.SetNumZeroed(ActivityCount);
+		for (int32 ActivityType = 0; ActivityType < ActivityCount; ++ActivityType)
+		{
+			const int64 ActivityId = ActivityMappings[ActivityType];
+			const TOptional<FActivityNodeID> NodeID = Graph.FindNodeByActivity(ActivityId);
+			if (!NodeID.IsSet())
+			{
+				Test.AddError(FString::Printf(TEXT("No node generated for activity %s"), *LexToString(static_cast<ETestActivity>(ActivityType))));
+				continue;
+			}
+			
+			const FActivityNode& Node = Graph.GetNodeById(*NodeID);
+			if (!NodeID.IsSet())
+			{
+				Test.AddError(FString::Printf(TEXT("Graph has invalid state. Node ID %lld is invalid for activity %s"), NodeID->ID, *LexToString(static_cast<ETestActivity>(ActivityType))));
+				continue;
+			}
+			
+			ActivityNodes[ActivityType] = *NodeID;
+		}
+
+		return ActivityNodes;
+	}
+}
+
+namespace UE::ConcertSyncTests::DeletingAndRecreatingActorIsHardDependency
+{
+	enum ETestActivity
+	{
+		CreateActor,
+		DeleteActor,
+		RecreateActor,
+
+		ActivityCount
+	};
+	
+	TArray<FActivityID> FillDatabase(FScopedSessionDatabase& Database);
+	FConcertExportedObject CreateActorMetaData(FName OuterLevelPath);
+	
+	/**
+	 * 1. Create actor A
+	 * 2. Delete Actor A
+	 * 3. Re-create actor A.
+	 *
+	 * 3 > 2 is a hard dependency (removing 2 would result in attempting to create the actor twice, which is invalid).
+	 */
+	IMPLEMENT_SIMPLE_AUTOMATION_TEST(FDeletingAndRecreatingActorIsHardDependency, "Concert.History.BuildGraph.DeletingAndRecreatingActorIsHardDependency", EAutomationTestFlags::ApplicationContextMask | EAutomationTestFlags::SmokeFilter)
+	bool FDeletingAndRecreatingActorIsHardDependency::RunTest(const FString& Parameters)
+	{
+		using namespace ConcertSyncCore;
+		
+		FScopedSessionDatabase SessionDatabase(*this);
+		const TArray<FActivityID> TestActivities = FillDatabase(SessionDatabase);
+
+		const FActivityDependencyGraph DependencyGraph = BuildDependencyGraphFrom(SessionDatabase);
+		UE_LOG(LogConcert, Log, TEXT("%s tested graph in Graphviz format:\n\n%s"), *GetTestFullName(), *ConcertSyncCore::Graphviz::ExportToGraphviz(DependencyGraph, SessionDatabase));
+
+		const TOptional<FActivityNodeID> CreateActorNodeID = DependencyGraph.FindNodeByActivity(TestActivities[CreateActor]);
+		const TOptional<FActivityNodeID> DeleteActorNodeID = DependencyGraph.FindNodeByActivity(TestActivities[DeleteActor]);
+		const TOptional<FActivityNodeID> RecreateActorNodeID = DependencyGraph.FindNodeByActivity(TestActivities[RecreateActor]);
+		if (!CreateActorNodeID || !DeleteActorNodeID || !RecreateActorNodeID)
+		{
+			AddError(TEXT("Activities not registered"));
+			return false;
+		}
+
+		const FActivityNode& CreatedActorNode = DependencyGraph.GetNodeById(*CreateActorNodeID);
+		const FActivityNode& DeleteActorNode = DependencyGraph.GetNodeById(*DeleteActorNodeID);
+		const FActivityNode& RecreateActorNode = DependencyGraph.GetNodeById(*RecreateActorNodeID);
+
+		TestEqual(TEXT("CreatedActorNode->GetDependencies().Num() == 0"), CreatedActorNode.GetDependencies().Num(), 0);
+		TestEqual(TEXT("DeleteActorNode->GetDependencies().Num() == 1"), DeleteActorNode.GetDependencies().Num(), 1);
+		TestEqual(TEXT("RecreateActorNode->GetDependencies().Num() == 1"), RecreateActorNode.GetDependencies().Num(), 1);
+
+		TestTrue(TEXT("DeleteActorNode depends on CreatedActorNode"), DeleteActorNode.DependsOnActivity(TestActivities[CreateActor], DependencyGraph, EActivityDependencyReason::SubobjectCreation, EDependencyStrength::HardDependency));
+		TestTrue(TEXT("RecreateActorNode depends on DeleteActorNode"), RecreateActorNode.DependsOnActivity(TestActivities[DeleteActor], DependencyGraph, EActivityDependencyReason::SubobjectRemoval, EDependencyStrength::HardDependency));
+
+		return true;
+	}
+	
+	TArray<FActivityID> FillDatabase(FScopedSessionDatabase& SessionDatabase)
+	{
+		TArray<FActivityID> ActivityIDs;
+		ActivityIDs.SetNumUninitialized(ActivityCount);
+
+		const FName FooLevel = TEXT("/Game/Foo");
+		int64 Dummy;
+		
+		// 1 Create actor
+		{
+			FConcertSyncTransactionActivity CreateActorActivity;
+			CreateActorActivity.EndpointId = SessionDatabase.GetEndpoint();
+			CreateActorActivity.EventData.Transaction.TransactionId = FGuid::NewGuid();
+			CreateActorActivity.EventData.Transaction.OperationId = FGuid::NewGuid();
+			FConcertExportedObject Actor = CreateActorMetaData(FooLevel);
+			Actor.ObjectData.bAllowCreate = true;
+			CreateActorActivity.EventData.Transaction.ExportedObjects = { Actor };
+			CreateActorActivity.EventData.Transaction.ModifiedPackages = { FooLevel };
+			SessionDatabase.GetTransactionMaxEventId(CreateActorActivity.EventId);
+			SessionDatabase.AddTransactionActivity(CreateActorActivity, ActivityIDs[CreateActor], Dummy);
+		}
+		
+		// 2 Delete actor
+		{
+			FConcertSyncTransactionActivity RemoveActorActivity;
+			RemoveActorActivity.EndpointId = SessionDatabase.GetEndpoint();
+			RemoveActorActivity.EventData.Transaction.TransactionId = FGuid::NewGuid();
+			RemoveActorActivity.EventData.Transaction.OperationId = FGuid::NewGuid();
+			FConcertExportedObject Actor = CreateActorMetaData(FooLevel);
+			Actor.ObjectData.bIsPendingKill = true;
+			RemoveActorActivity.EventData.Transaction.ExportedObjects = { Actor };
+			RemoveActorActivity.EventData.Transaction.ModifiedPackages = { FooLevel };
+			SessionDatabase.GetTransactionMaxEventId(RemoveActorActivity.EventId);
+			SessionDatabase.AddTransactionActivity(RemoveActorActivity, ActivityIDs[DeleteActor], Dummy);
+		}
+		
+		// 3 Re-create actor
+		{
+			FConcertSyncTransactionActivity CreateActorActivity;
+			CreateActorActivity.EndpointId = SessionDatabase.GetEndpoint();
+			CreateActorActivity.EventData.Transaction.TransactionId = FGuid::NewGuid();
+			CreateActorActivity.EventData.Transaction.OperationId = FGuid::NewGuid();
+			FConcertExportedObject Actor = CreateActorMetaData(FooLevel);
+			Actor.ObjectData.bAllowCreate = true;
+			CreateActorActivity.EventData.Transaction.ExportedObjects = { Actor };
+			CreateActorActivity.EventData.Transaction.ModifiedPackages = { FooLevel };
+			SessionDatabase.GetTransactionMaxEventId(CreateActorActivity.EventId);
+			SessionDatabase.AddTransactionActivity(CreateActorActivity, ActivityIDs[RecreateActor], Dummy);
+		}
+		
+		return ActivityIDs;
+	}
+	
+	FConcertExportedObject CreateActorMetaData(FName OuterLevelPath)
+	{
+		FConcertExportedObject Result;
+		Result.ObjectId.ObjectName = TEXT("SomeTestActorName42");
+		Result.ObjectId.ObjectPackageName = OuterLevelPath;
+		Result.ObjectId.ObjectOuterPathName = *FString::Printf(TEXT("%s:PersistentLevel"), *OuterLevelPath.ToString());
+		Result.ObjectId.ObjectClassPathName = TEXT("/Script/Engine.StaticMeshActor");
+		return Result;
+	}
+}
+
+namespace UE::ConcertSyncTests::IndirectPackageDependencyTest
+{
+	/**
+	 * This tests that potential indirect dependencies are handled.
+	 *
+	 * Sequence of user actions:
+	 *  1 Create data asset A
+	 *  2 Make actor reference A
+	 *  3 Edit data asset
+	 *  4 Edit actor
+	 *
+	 * The dependency graph should look like this:
+	 *  2 -> 1 (PackageCreation)
+	 *  3 -> 1 (PackageCreation)
+	 *  4 -> 1 (EditPossiblyDependsOnPackage)
+	 *  4 -> 2 (EditPossiblyDependsOnPackage)
+	 *
+	 * This is relevant because the actor's construction script may depend query data from the data asset.
+	 */
+	IMPLEMENT_SIMPLE_AUTOMATION_TEST(FIndirectPackageDependencyTest, "Concert.History.BuildGraph.IndirectPackageDependencyTest", EAutomationTestFlags::ApplicationContextMask | EAutomationTestFlags::SmokeFilter)
+	bool FIndirectPackageDependencyTest::RunTest(const FString& Parameters)
+	{
+		// TODO:
+		return true;
+	}
+}
\ No newline at end of file
diff --git a/Engine/Plugins/Developer/Concert/ConcertSync/ConcertSyncTest/Source/ConcertSyncTest/Private/HistoryEdition/DependencyGraphTests.cpp b/Engine/Plugins/Developer/Concert/ConcertSync/ConcertSyncTest/Source/ConcertSyncTest/Private/HistoryEdition/DependencyGraphTests.cpp
new file mode 100644
--- /dev/null
+++ b/Engine/Plugins/Developer/Concert/ConcertSync/ConcertSyncTest/Source/ConcertSyncTest/Private/HistoryEdition/DependencyGraphTests.cpp
@@ -0,0 +1,73 @@
+// Copyright Epic Games, Inc. All Rights Reserved.
+
+#include "HistoryEdition/ActivityDependencyGraph.h"
+
+#include "Misc/AutomationTest.h"
+
+namespace UE::ConcertSyncTests
+{
+	/** Ensures that the dependency graph API functions as intended, i.e. you add, connect and query nodes. */
+	IMPLEMENT_SIMPLE_AUTOMATION_TEST(FConcertDependencyGraphApiTest, "Concert.History.GraphApi.DependencyGraphApi", EAutomationTestFlags::ApplicationContextMask | EAutomationTestFlags::SmokeFilter)
+	bool FConcertDependencyGraphApiTest::RunTest(const FString& Parameters)
+	{
+		using namespace UE::ConcertSyncCore;
+		constexpr int64 FirstActivity = 1;
+		constexpr int64 SecondActivity = 20;
+		constexpr int64 ThirdActivity = 300;
+		
+		FActivityDependencyGraph Graph;
+		const FActivityNodeID FirstID = Graph.AddActivity(FirstActivity);
+		const FActivityNodeID SecondID = Graph.AddActivity(SecondActivity);
+		const FActivityNodeID ThirdID = Graph.AddActivity(ThirdActivity);
+
+		const FActivityNode& FirstNode = Graph.GetNodeById(FirstID);
+		const FActivityNode& SecondNode = Graph.GetNodeById(SecondID);
+		const FActivityNode& ThirdNode = Graph.GetNodeById(ThirdID);
+
+		TestTrue(TEXT("Correct activity in Root"), FirstNode.GetActivityId() == FirstActivity);
+		TestTrue(TEXT("Correct activity in LeftNode"), SecondNode.GetActivityId() == SecondActivity);
+		TestTrue(TEXT("Correct activity in RightNode"), ThirdNode.GetActivityId() == ThirdActivity);
+
+		const FActivityNodeID Invalid{42342343};
+		TestEqual(TEXT("Find activity: First"), Graph.FindNodeByActivity(FirstActivity).Get(Invalid), FirstID);
+		TestEqual(TEXT("Find activity: Second"), Graph.FindNodeByActivity(SecondActivity).Get(Invalid), SecondID);
+		TestEqual(TEXT("Find activity: Third"), Graph.FindNodeByActivity(ThirdActivity).Get(Invalid), ThirdID);
+
+		const bool bLefttoRoot = Graph.AddDependency(SecondID, { FirstID, EActivityDependencyReason::PackageCreation, EDependencyStrength::HardDependency });
+		const bool bRightToRoot = Graph.AddDependency(ThirdID, { FirstID, EActivityDependencyReason::PackageCreation, EDependencyStrength::PossibleDependency });
+
+		TestTrue(TEXT("Added left-to-root"), bLefttoRoot);
+		TestTrue(TEXT("Added right-to-root"), bRightToRoot);
+		if (!bLefttoRoot || !bRightToRoot)
+		{
+			return false;
+		}
+
+		TestTrue(TEXT("1st knows about 2nd activity"), FirstNode.AffectsAnyActivity() && FirstNode.AffectsNode(SecondID) && FirstNode.AffectsActivity(SecondActivity, Graph));
+		TestTrue(TEXT("1st knows about 3rd activilty"), FirstNode.AffectsAnyActivity() && FirstNode.AffectsNode(ThirdID) && FirstNode.AffectsActivity(ThirdActivity, Graph));
+		TestTrue(TEXT("2nd knows about 1st"), SecondNode.HasAnyDependency() && SecondNode.DependsOnNode(FirstID, EActivityDependencyReason::PackageCreation, EDependencyStrength::HardDependency) && SecondNode.DependsOnActivity(FirstActivity, Graph, EActivityDependencyReason::PackageCreation, EDependencyStrength::HardDependency));
+		TestTrue(TEXT("3rd knows about 1st"), ThirdNode.HasAnyDependency() && ThirdNode.DependsOnNode(FirstID, EActivityDependencyReason::PackageCreation, EDependencyStrength::PossibleDependency) && ThirdNode.DependsOnActivity(FirstActivity, Graph, EActivityDependencyReason::PackageCreation, EDependencyStrength::PossibleDependency));
+		return true;
+	}
+
+	/** Ensures that the dependency graph API functions as intended, i.e. you add, connect and query nodes. */
+	IMPLEMENT_SIMPLE_AUTOMATION_TEST(FDisallowCyclesInGraph, "Concert.History.GraphApi.DisallowCyclesInGraph", EAutomationTestFlags::ApplicationContextMask | EAutomationTestFlags::SmokeFilter)
+	bool FDisallowCyclesInGraph::RunTest(const FString& Parameters)
+	{
+		using namespace UE::ConcertSyncCore;
+		constexpr int64 FirstActivity = 1;
+		constexpr int64 SecondActivity = 20;
+		
+		FActivityDependencyGraph Graph;
+		const FActivityNodeID FirstNodeID = Graph.AddActivity(FirstActivity);
+		const FActivityNodeID SecondNodeID = Graph.AddActivity(SecondActivity);
+
+		const bool bSecondToFirst = Graph.AddDependency(SecondNodeID, { FirstNodeID, EActivityDependencyReason::PackageCreation, EDependencyStrength::HardDependency });
+		AddExpectedError(TEXT("Activities can only depend on preceding activities!"), EAutomationExpectedErrorFlags::Contains, 1);
+		const bool bFirstToSecond = Graph.AddDependency(FirstNodeID, { SecondNodeID, EActivityDependencyReason::PackageCreation, EDependencyStrength::PossibleDependency });
+
+		TestTrue(TEXT("Later activity depends on early dependency"), bSecondToFirst);
+		TestTrue(TEXT("Earlier dependency cannot depend on later dependency"), !bFirstToSecond);
+		return true;
+	}
+}
\ No newline at end of file
diff --git a/Engine/Plugins/Developer/Concert/ConcertSync/ConcertSyncTest/Source/ConcertSyncTest/Private/HistoryEdition/RenameEditAndDeleteMapsFlow.cpp b/Engine/Plugins/Developer/Concert/ConcertSync/ConcertSyncTest/Source/ConcertSyncTest/Private/HistoryEdition/RenameEditAndDeleteMapsFlow.cpp
new file mode 100644
--- /dev/null
+++ b/Engine/Plugins/Developer/Concert/ConcertSync/ConcertSyncTest/Source/ConcertSyncTest/Private/HistoryEdition/RenameEditAndDeleteMapsFlow.cpp
@@ -0,0 +1,218 @@
+﻿// Copyright Epic Games, Inc. All Rights Reserved.
+
+#include "RenameEditAndDeleteMapsFlow.h"
+
+#include "ConcertLogGlobal.h"
+#include "ConcertSyncSessionDatabase.h"
+#include "ScopedSessionDatabase.h"
+
+namespace UE::ConcertSyncTests::RenameEditAndDeleteMapsFlowTest
+{
+	template<typename T>
+	T MakeActivity(const FGuid& EndpointID)
+	{
+		T Activity;
+		Activity.EndpointId = EndpointID;
+		return Activity;
+	}
+
+	TSet<ETestActivity> AllActivities()
+	{
+		return {
+			_1_NewPackageFoo,
+			_1_SavePackageFoo,
+			_2_AddActor,
+			_3_RenameActor,
+			_4_EditActor,
+			_5_SavePackageBar,
+			_5_RenameFooToBar,
+			_6_EditActor,
+			_7_DeleteBar,
+			_8_NewPackageFoo,
+			_8_SavePackageFoo
+		};
+	}
+
+	FString LexToString(ETestActivity Activity)
+	{
+		switch (Activity)
+		{
+		case _1_NewPackageFoo: return TEXT("\"1 New package Foo\"");
+		case _1_SavePackageFoo: return TEXT("\"1 Saved package Foo\"");
+		case _2_AddActor: return TEXT("\"2 Create actor\"");
+		case _3_RenameActor: return TEXT("\"3 Edit actor\"");
+		case _4_EditActor: return TEXT("\"4 Edit actor\"");
+		case _5_SavePackageBar: return TEXT("\"5 Save package\"");
+		case _5_RenameFooToBar: return TEXT("\"5 Rename Foo to Bar\"");
+		case _6_EditActor: return TEXT("\"6 Edit actor\"");
+		case _7_DeleteBar: return TEXT("\"7 Delete package Bar\"");
+		case _8_NewPackageFoo: return TEXT("\"8 Create package Bar\"");
+		case _8_SavePackageFoo: return TEXT("\"8 Save package Bar\"");
+				
+		case ActivityCount:
+		default:
+			checkNoEntry();
+			return FString();
+		}
+	}
+
+	struct FCreatedStaticMeshActor
+	{
+		FConcertExportedObject Actor;
+		FConcertExportedObject StaticMeshComponent;
+	};
+	
+	FCreatedStaticMeshActor CreateEditedActor(FName OuterLevelPath);
+
+	TTestActivityArray<int64> CreateActivityHistory(FConcertSyncSessionDatabase& SessionDatabase, const FGuid& EndpointID)
+	{
+		TTestActivityArray<int64> ActivityIDs;
+		TTestActivityArray<int64> PackageEventIDs;
+		ActivityIDs.SetNumUninitialized(ActivityCount);
+		PackageEventIDs.SetNumUninitialized(ActivityCount);
+		
+		bool bAllSucceeded = true;
+		
+		// The names of the activities make it into the generated Graphviz graph 
+		const FName EditedActorName = TEXT("Actor");
+		const FName FooLevel = TEXT("/Game/Foo");
+		const FName BarLevel = TEXT("/Game/Bar");
+		
+		// 1 Create map Foo
+		{
+			FConcertSyncActivity NewPackage = MakeActivity<FConcertSyncActivity>(EndpointID);
+			FConcertPackageInfo PackageInfo;
+			FConcertPackageDataStream PackageDataStream;
+			PackageInfo.PackageName = FooLevel;
+			PackageInfo.PackageUpdateType = EConcertPackageUpdateType::Added;
+			SessionDatabase.GetTransactionMaxEventId(PackageInfo.TransactionEventIdAtSave);
+			bAllSucceeded &= SessionDatabase.AddPackageActivity(NewPackage, PackageInfo, PackageDataStream, ActivityIDs[_1_NewPackageFoo], PackageEventIDs[_1_NewPackageFoo]);
+
+			FConcertSyncActivity SavePackage = MakeActivity<FConcertSyncActivity>(EndpointID);
+			PackageInfo.PackageUpdateType = EConcertPackageUpdateType::Saved;
+			SessionDatabase.GetTransactionMaxEventId(PackageInfo.TransactionEventIdAtSave);
+			bAllSucceeded &= SessionDatabase.AddPackageActivity(SavePackage, PackageInfo, PackageDataStream, ActivityIDs[_1_SavePackageFoo], PackageEventIDs[_1_SavePackageFoo]);
+		}
+
+		// 2 Add actor A
+		{
+			FConcertSyncTransactionActivity CreateActor = MakeActivity<FConcertSyncTransactionActivity>(EndpointID);
+			CreateActor.EventData.Transaction.TransactionId = FGuid::NewGuid();
+			CreateActor.EventData.Transaction.OperationId = FGuid::NewGuid();
+			FCreatedStaticMeshActor NewActorData = CreateEditedActor(FooLevel);
+			NewActorData.Actor.ObjectData.bAllowCreate = true;
+			NewActorData.StaticMeshComponent.ObjectData.bAllowCreate = true;
+			CreateActor.EventData.Transaction.ExportedObjects = { NewActorData.StaticMeshComponent, NewActorData.Actor };
+			CreateActor.EventData.Transaction.ModifiedPackages = { FooLevel };
+			SessionDatabase.GetTransactionMaxEventId(CreateActor.EventId);
+			bAllSucceeded &= SessionDatabase.AddTransactionActivity(CreateActor, ActivityIDs[_2_AddActor], PackageEventIDs[_2_AddActor]);
+		}
+
+		// 3 Rename actor A
+		{
+			FConcertSyncTransactionActivity EditActor = MakeActivity<FConcertSyncTransactionActivity>(EndpointID);
+			EditActor.EventData.Transaction.TransactionId = FGuid::NewGuid();
+			EditActor.EventData.Transaction.OperationId = FGuid::NewGuid();
+			
+			FCreatedStaticMeshActor NewActorData = CreateEditedActor(FooLevel);
+			NewActorData.Actor.PropertyDatas = { { TEXT("ActorLabel"), {} } };
+			EditActor.EventData.Transaction.ExportedObjects = { NewActorData.Actor };
+			EditActor.EventData.Transaction.ModifiedPackages = { FooLevel };
+			SessionDatabase.GetTransactionMaxEventId(EditActor.EventId);
+			bAllSucceeded &= SessionDatabase.AddTransactionActivity(EditActor, ActivityIDs[_3_RenameActor], PackageEventIDs[_3_RenameActor]);
+		}
+
+		// 4 Edit actor A
+		{
+			FConcertSyncTransactionActivity EditActor = MakeActivity<FConcertSyncTransactionActivity>(EndpointID);
+			EditActor.EventData.Transaction.TransactionId = FGuid::NewGuid();
+			EditActor.EventData.Transaction.OperationId = FGuid::NewGuid();
+			FCreatedStaticMeshActor NewActorData = CreateEditedActor(FooLevel);
+			NewActorData.StaticMeshComponent.PropertyDatas = { { TEXT("RelativeLocation"), {} } };
+			EditActor.EventData.Transaction.ExportedObjects = { NewActorData.StaticMeshComponent };
+			EditActor.EventData.Transaction.ModifiedPackages = { FooLevel };
+			SessionDatabase.GetTransactionMaxEventId(EditActor.EventId);
+			bAllSucceeded &= SessionDatabase.AddTransactionActivity(EditActor, ActivityIDs[_4_EditActor], PackageEventIDs[_4_EditActor]);
+		}
+
+		// 5 Rename map to Bar
+		{
+			FConcertSyncActivity SaveBarPackage = MakeActivity<FConcertSyncActivity>(EndpointID);
+			FConcertPackageInfo PackageInfo;
+			FConcertPackageDataStream PackageDataStream;
+			PackageInfo.PackageName = BarLevel;
+			PackageInfo.PackageUpdateType = EConcertPackageUpdateType::Saved;
+			SessionDatabase.GetTransactionMaxEventId(PackageInfo.TransactionEventIdAtSave);
+			bAllSucceeded &= SessionDatabase.AddPackageActivity(SaveBarPackage, PackageInfo, PackageDataStream, ActivityIDs[_5_SavePackageBar], PackageEventIDs[_5_SavePackageBar]);
+
+			PackageInfo.PackageUpdateType = EConcertPackageUpdateType::Renamed;
+			PackageInfo.PackageName = FooLevel;
+			PackageInfo.NewPackageName = BarLevel;
+			bAllSucceeded &= SessionDatabase.AddPackageActivity(SaveBarPackage, PackageInfo, PackageDataStream, ActivityIDs[_5_RenameFooToBar], PackageEventIDs[_5_RenameFooToBar]);
+		}
+
+		// 6 Edit actor A
+		{
+			FConcertSyncTransactionActivity EditActor = MakeActivity<FConcertSyncTransactionActivity>(EndpointID);
+			EditActor.EventData.Transaction.TransactionId = FGuid::NewGuid();
+			EditActor.EventData.Transaction.OperationId = FGuid::NewGuid();
+			FCreatedStaticMeshActor NewActorData = CreateEditedActor(BarLevel);
+			NewActorData.StaticMeshComponent.PropertyDatas = { { TEXT("RelativeLocation"), {} } };
+			EditActor.EventData.Transaction.ExportedObjects = { NewActorData.StaticMeshComponent };
+			EditActor.EventData.Transaction.ModifiedPackages = { BarLevel };
+			SessionDatabase.GetTransactionMaxEventId(EditActor.EventId);
+			bAllSucceeded &= SessionDatabase.AddTransactionActivity(EditActor, ActivityIDs[_6_EditActor], PackageEventIDs[_6_EditActor]);
+		}
+
+		// 7 Delete map Bar
+		{
+			FConcertSyncActivity SaveBarPackage = MakeActivity<FConcertSyncActivity>(EndpointID);
+			FConcertPackageInfo PackageInfo;
+			FConcertPackageDataStream PackageDataStream;
+			PackageInfo.PackageName = BarLevel;
+			PackageInfo.PackageUpdateType = EConcertPackageUpdateType::Deleted;
+			SessionDatabase.GetTransactionMaxEventId(PackageInfo.TransactionEventIdAtSave);
+			bAllSucceeded &= SessionDatabase.AddPackageActivity(SaveBarPackage, PackageInfo, PackageDataStream, ActivityIDs[_7_DeleteBar], PackageEventIDs[_7_DeleteBar]);
+		}
+		
+		// 8 Create map Bar
+		{
+			FConcertSyncActivity NewPackage = MakeActivity<FConcertSyncActivity>(EndpointID);
+			FConcertPackageInfo PackageInfo;
+			FConcertPackageDataStream PackageDataStream;
+			PackageInfo.PackageName = TEXT("/Game/Bar");
+			PackageInfo.PackageUpdateType = EConcertPackageUpdateType::Added;
+			SessionDatabase.GetTransactionMaxEventId(PackageInfo.TransactionEventIdAtSave);
+			bAllSucceeded &= SessionDatabase.AddPackageActivity(NewPackage, PackageInfo, PackageDataStream, ActivityIDs[_8_NewPackageFoo], PackageEventIDs[_8_NewPackageFoo]);
+
+			FConcertSyncActivity SavePackage = MakeActivity<FConcertSyncActivity>(EndpointID);
+			PackageInfo.PackageUpdateType = EConcertPackageUpdateType::Saved;
+			SessionDatabase.GetTransactionMaxEventId(PackageInfo.TransactionEventIdAtSave);
+			bAllSucceeded &= SessionDatabase.AddPackageActivity(SavePackage, PackageInfo, PackageDataStream, ActivityIDs[_8_SavePackageFoo], PackageEventIDs[_8_SavePackageFoo]);
+		}
+
+		if (!bAllSucceeded)
+		{
+			UE_LOG(LogConcert, Error, TEXT("Something went wrong creating the activities. Test result may be wrong."))
+		}
+		return ActivityIDs;
+	}
+
+	FCreatedStaticMeshActor CreateEditedActor(FName OuterLevelPath)
+	{
+		FCreatedStaticMeshActor Result;
+
+		Result.Actor.ObjectId.ObjectName = TEXT("StaticMeshActor0");
+		Result.Actor.ObjectId.ObjectPackageName = OuterLevelPath;
+		Result.Actor.ObjectId.ObjectOuterPathName = *FString::Printf(TEXT("%s:PersistentLevel"), *OuterLevelPath.ToString());
+		Result.Actor.ObjectId.ObjectClassPathName = TEXT("/Script/Engine.StaticMeshActor");
+		
+		Result.StaticMeshComponent.ObjectId.ObjectName = TEXT("StaticMeshComponent0");
+		Result.StaticMeshComponent.ObjectId.ObjectPackageName = OuterLevelPath;
+		Result.StaticMeshComponent.ObjectId.ObjectOuterPathName = *FString::Printf(TEXT("%s:PersistentLevel.%s"), *OuterLevelPath.ToString(), *Result.Actor.ObjectId.ObjectName.ToString());
+		Result.StaticMeshComponent.ObjectId.ObjectClassPathName = TEXT("/Script/Engine.StaticMeshComponent");
+		
+		return Result;
+	}
+}
+
diff --git a/Engine/Plugins/Developer/Concert/ConcertSync/ConcertSyncTest/Source/ConcertSyncTest/Private/HistoryEdition/RenameEditAndDeleteMapsFlow.h b/Engine/Plugins/Developer/Concert/ConcertSync/ConcertSyncTest/Source/ConcertSyncTest/Private/HistoryEdition/RenameEditAndDeleteMapsFlow.h
new file mode 100644
--- /dev/null
+++ b/Engine/Plugins/Developer/Concert/ConcertSync/ConcertSyncTest/Source/ConcertSyncTest/Private/HistoryEdition/RenameEditAndDeleteMapsFlow.h
@@ -0,0 +1,48 @@
+﻿// Copyright Epic Games, Inc. All Rights Reserved.
+
+#pragma once
+
+#include "CoreMinimal.h"
+
+class FConcertSyncSessionDatabase;
+
+namespace UE::ConcertSyncTests::RenameEditAndDeleteMapsFlowTest
+{
+	enum ETestActivity
+	{
+		_1_NewPackageFoo,
+		_1_SavePackageFoo,
+		_2_AddActor,
+		_3_RenameActor,
+		_4_EditActor,
+		_5_SavePackageBar,
+		_5_RenameFooToBar,
+		_6_EditActor,
+		_7_DeleteBar,
+		_8_NewPackageFoo,
+		_8_SavePackageFoo,
+		
+		ActivityCount
+	};
+	TSet<ETestActivity> AllActivities();
+	FString LexToString(ETestActivity Activity);
+	
+	/** An array where every entry of ETestActivity is a valid index. */
+	template<typename T>
+	using TTestActivityArray = TArray<T, TInlineAllocator<ActivityCount>>;
+
+	/**
+	 * Creates a session history which resembles the following sequence of user actions:
+	 *	1 Create map Foo
+	 *	2 Add actor A
+	 *	3 Edit actor A
+	 *	4 Edit actor A
+	 *	5 Rename map to Bar
+	 *	6 Edit actor A
+	 *	7 Delete map Bar
+	 *	8 Create map Bar
+	 *
+	 *	@return An array where every ETestActivity entry is the index to the activity ID added to SessionDatabase
+	 */
+	TTestActivityArray<int64> CreateActivityHistory(FConcertSyncSessionDatabase& SessionDatabase, const FGuid& EndpointID);
+}
\ No newline at end of file
diff --git a/Engine/Plugins/Developer/Concert/ConcertSync/ConcertSyncTest/Source/ConcertSyncTest/Private/HistoryEdition/ScopedSessionDatabase.h b/Engine/Plugins/Developer/Concert/ConcertSync/ConcertSyncTest/Source/ConcertSyncTest/Private/HistoryEdition/ScopedSessionDatabase.h
new file mode 100644
--- /dev/null
+++ b/Engine/Plugins/Developer/Concert/ConcertSync/ConcertSyncTest/Source/ConcertSyncTest/Private/HistoryEdition/ScopedSessionDatabase.h
@@ -0,0 +1,44 @@
+﻿// Copyright Epic Games, Inc. All Rights Reserved.
+
+#pragma once
+
+#include "CoreMinimal.h"
+#include "ConcertSyncSessionDatabase.h"
+
+#include "Misc/AutomationTest.h"
+
+namespace UE::ConcertSyncTests
+{
+	class FScopedSessionDatabase : public FConcertSyncSessionDatabase
+	{
+		FAutomationTestBase& Test;
+		const FString TestSessionPath_Server = FPaths::ProjectIntermediateDir() / TEXT("ConcertDatabaseTest_Server");
+		const FGuid EndpointID;
+	public:
+
+		FScopedSessionDatabase(FAutomationTestBase& Test)
+			: Test(Test)
+			, EndpointID(FGuid::NewGuid())
+		{
+			Open(TestSessionPath_Server);
+
+			FConcertSyncEndpointData EndpointData;
+			EndpointData.ClientInfo.Initialize();
+			if (!SetEndpoint(EndpointID, EndpointData))
+			{
+				Test.AddError(FString::Printf(TEXT("Test may be faulty because endpoint could not be set: %s"), *GetLastError()));
+			}
+		}
+
+		~FScopedSessionDatabase()
+		{
+			if (IsValid() && !Close())
+			{
+				Test.AddError(TEXT("Failed to close server database"));
+			}
+			IFileManager::Get().DeleteDirectory(*TestSessionPath_Server, false, true);
+		}
+
+		const FGuid& GetEndpoint() const { return EndpointID; }
+	};
+}
\ No newline at end of file
