Fix increment build issue with Mac

diff --git a/Engine/Plugins/Developer/Concert/ConcertSync/ConcertSyncCore/Source/ConcertSyncCore/Public/HistoryEdition/ActivityGraphIDs.h b/Engine/Plugins/Developer/Concert/ConcertSync/ConcertSyncCore/Source/ConcertSyncCore/Public/HistoryEdition/ActivityGraphIDs.h
--- a/Engine/Plugins/Developer/Concert/ConcertSync/ConcertSyncCore/Source/ConcertSyncCore/Public/HistoryEdition/ActivityGraphIDs.h
+++ b/Engine/Plugins/Developer/Concert/ConcertSync/ConcertSyncCore/Source/ConcertSyncCore/Public/HistoryEdition/ActivityGraphIDs.h
@@ -7,7 +7,6 @@
 
 namespace UE::ConcertSyncCore
 {
-	class FActivityDependecyGraph;
 	/** Enforces type-safe usage of IDs by avoiding implicit conversions. */
 	struct CONCERTSYNCCORE_API FActivityNodeID
 	{
@@ -24,8 +23,6 @@ namespace UE::ConcertSyncCore
 			: ID(Other.ID)
 		{}
 
-		bool HasAnyDependency(const FActivityDependecyGraph&) const;
-
 		/** Conversion back to to size_t is not a common programmer mistake */
 		explicit operator size_t() const { return ID; }
 
@@ -41,7 +38,12 @@ namespace UE::ConcertSyncCore
 
 		FORCEINLINE friend uint32 GetTypeHash(const FActivityNodeID& NodeID)
 		{
+#if PLATFORM_MAC
+			// Help Mac compiler correctly resolve the candidate function
+			return ::GetTypeHash(static_cast<uint64>(NodeID.ID));
+#else
 			return ::GetTypeHash(NodeID.ID);
+#endif
 		}
 	};
 }
\ No newline at end of file
