Fix CIS issues

diff --git a/Engine/Plugins/Developer/Concert/ConcertSync/ConcertSyncCore/Source/ConcertSyncCore/Private/HistoryEdition/HistoryAnalysis.cpp b/Engine/Plugins/Developer/Concert/ConcertSync/ConcertSyncCore/Source/ConcertSyncCore/Private/HistoryEdition/HistoryAnalysis.cpp
--- a/Engine/Plugins/Developer/Concert/ConcertSync/ConcertSyncCore/Source/ConcertSyncCore/Private/HistoryEdition/HistoryAnalysis.cpp
+++ b/Engine/Plugins/Developer/Concert/ConcertSync/ConcertSyncCore/Source/ConcertSyncCore/Private/HistoryEdition/HistoryAnalysis.cpp
@@ -2,8 +2,11 @@
 
 #include "HistoryEdition/HistoryAnalysis.h"
 
+#include "HistoryEdition/ActivityGraphIDs.h"
 #include "HistoryEdition/DependencyGraphBuilder.h"
 
+#include "Containers/Queue.h"
+
 namespace UE::ConcertSyncCore
 {
 	FHistoryDeletionRequirements AnalyseActivityDeletion(const TSet<FActivityID>& ActivitiesToDelete, const FConcertSyncSessionDatabase& Database, bool bAddActivitiesToDelete)
diff --git a/Engine/Plugins/Developer/Concert/ConcertSync/ConcertSyncTest/Source/ConcertSyncTest/Private/HistoryEdition/ScopedSessionDatabase.h b/Engine/Plugins/Developer/Concert/ConcertSync/ConcertSyncTest/Source/ConcertSyncTest/Private/HistoryEdition/ScopedSessionDatabase.h
--- a/Engine/Plugins/Developer/Concert/ConcertSync/ConcertSyncTest/Source/ConcertSyncTest/Private/HistoryEdition/ScopedSessionDatabase.h
+++ b/Engine/Plugins/Developer/Concert/ConcertSync/ConcertSyncTest/Source/ConcertSyncTest/Private/HistoryEdition/ScopedSessionDatabase.h
@@ -5,7 +5,9 @@
 #include "CoreMinimal.h"
 #include "ConcertSyncSessionDatabase.h"
 
+#include "HAL/FileManager.h"
 #include "Misc/AutomationTest.h"
+#include "Misc/Paths.h"
 
 namespace UE::ConcertSyncTests
 {
