Add function for deleting arbitrary activities from the session database

diff --git a/Engine/Plugins/Developer/Concert/ConcertMain/Source/Concert/Private/ConcertServer.cpp b/Engine/Plugins/Developer/Concert/ConcertMain/Source/Concert/Private/ConcertServer.cpp
--- a/Engine/Plugins/Developer/Concert/ConcertMain/Source/Concert/Private/ConcertServer.cpp
+++ b/Engine/Plugins/Developer/Concert/ConcertMain/Source/Concert/Private/ConcertServer.cpp
@@ -564,7 +564,7 @@ void FConcertServer::ArchiveOfflineSessions(const FConcertServerSessionRepositor
 	}
 }
 
-FGuid FConcertServer::ArchiveSession(const FGuid& SessionId, const FString& ArchiveNameOverride, const FConcertSessionFilter& SessionFilter, FText& OutFailureReason)
+FGuid FConcertServer::ArchiveSession(const FGuid& SessionId, const FString& ArchiveNameOverride, const FConcertSessionFilter& SessionFilter, FText& OutFailureReason, FGuid ArchiveSessionIdOverride)
 {
 	if (GetArchivedSessionIdByName(ArchiveNameOverride).IsValid())
 	{
@@ -572,7 +572,7 @@ FGuid FConcertServer::ArchiveSession(const FGuid& SessionId, const FString& Arch
 		return FGuid();
 	}
 
-	const FGuid ArchivedSessionId = ArchiveLiveSession(SessionId, ArchiveNameOverride, SessionFilter);
+	const FGuid ArchivedSessionId = ArchiveLiveSession(SessionId, ArchiveNameOverride, SessionFilter, MoveTemp(ArchiveSessionIdOverride));
 	if (!ArchivedSessionId.IsValid())
 	{
 		OutFailureReason = LOCTEXT("Error_ArchiveSession_FailedToCopy", "Could not copy session data to the archive");
@@ -1439,7 +1439,7 @@ bool FConcertServer::DestroyLiveSession(const FGuid& LiveSessionId, const bool b
 	return false;
 }
 
-FGuid FConcertServer::ArchiveLiveSession(const FGuid& LiveSessionId, const FString& ArchivedSessionNameOverride, const FConcertSessionFilter& SessionFilter)
+FGuid FConcertServer::ArchiveLiveSession(const FGuid& LiveSessionId, const FString& ArchivedSessionNameOverride, const FConcertSessionFilter& SessionFilter, FGuid ArchiveSessionIdOverride)
 {
 	TSharedPtr<IConcertServerSession> LiveSession = LiveSessions.FindRef(LiveSessionId);
 	if (LiveSession)
@@ -1458,7 +1458,7 @@ FGuid FConcertServer::ArchiveLiveSession(const FGuid& LiveSessionId, const FStri
 		const FConcertServerSessionRepository& SessionRepository = GetSessionRepository(LiveSession->GetSessionInfo().SessionId);
 
 		FConcertSessionInfo ArchivedSessionInfo = LiveSession->GetSessionInfo();
-		ArchivedSessionInfo.SessionId = FGuid::NewGuid();
+		ArchivedSessionInfo.SessionId = ensure(ArchiveSessionIdOverride.IsValid()) ? MoveTemp(ArchiveSessionIdOverride) : FGuid::NewGuid();
 		ArchivedSessionInfo.SessionName = MoveTemp(ArchivedSessionName);
 		if (EventSink->ArchiveSession(*this, LiveSession.ToSharedRef(), SessionRepository.GetSessionSavedDir(ArchivedSessionInfo.SessionId), ArchivedSessionInfo, SessionFilter))
 		{
diff --git a/Engine/Plugins/Developer/Concert/ConcertMain/Source/Concert/Private/ConcertServer.h b/Engine/Plugins/Developer/Concert/ConcertMain/Source/Concert/Private/ConcertServer.h
--- a/Engine/Plugins/Developer/Concert/ConcertMain/Source/Concert/Private/ConcertServer.h
+++ b/Engine/Plugins/Developer/Concert/ConcertMain/Source/Concert/Private/ConcertServer.h
@@ -42,7 +42,7 @@ public:
 	virtual TSharedPtr<IConcertServerSession> CreateSession(const FConcertSessionInfo& SessionInfo, FText& OutFailureReason) override;
 	virtual TSharedPtr<IConcertServerSession> RestoreSession(const FGuid& SessionId, const FConcertSessionInfo& SessionInfo, const FConcertSessionFilter& SessionFilter, FText& OutFailureReason) override;
 	virtual TSharedPtr<IConcertServerSession> CopySession(const FGuid& SrcSessionId, const FConcertSessionInfo& NewSessionInfo, const FConcertSessionFilter& SessionFilter, FText& OutFailureReason) override;
-	virtual FGuid ArchiveSession(const FGuid& SessionId, const FString& ArchiveNameOverride, const FConcertSessionFilter& SessionFilter, FText& OutFailureReason) override;
+	virtual FGuid ArchiveSession(const FGuid& SessionId, const FString& ArchiveNameOverride, const FConcertSessionFilter& SessionFilter, FText& OutFailureReason, FGuid ArchiveSessionIdOverride = FGuid::NewGuid()) override;
 	virtual bool ExportSession(const FGuid& SessionId, const FConcertSessionFilter& SessionFilter, const FString& DestDir, bool bAnonymizeData, FText& OutFailureReason) override;
 	virtual bool RenameSession(const FGuid& SessionId, const FString& NewName, FText& OutFailureReason) override;
 	virtual bool DestroySession(const FGuid& SessionId, FText& OutFailureReason) override;
@@ -135,7 +135,7 @@ private:
 	bool DestroyLiveSession(const FGuid& LiveSessionId, const bool bDeleteSessionData);
 
 	/**  */
-	FGuid ArchiveLiveSession(const FGuid& LiveSessionId, const FString& ArchivedSessionNameOverride, const FConcertSessionFilter& SessionFilter);
+	FGuid ArchiveLiveSession(const FGuid& LiveSessionId, const FString& ArchivedSessionNameOverride, const FConcertSessionFilter& SessionFilter, FGuid ArchiveSessionIdOverride = FGuid::NewGuid());
 
 	/**  */
 	bool CreateArchivedSession(const FConcertSessionInfo& SessionInfo);
diff --git a/Engine/Plugins/Developer/Concert/ConcertMain/Source/Concert/Public/IConcertServer.h b/Engine/Plugins/Developer/Concert/ConcertMain/Source/Concert/Public/IConcertServer.h
--- a/Engine/Plugins/Developer/Concert/ConcertMain/Source/Concert/Public/IConcertServer.h
+++ b/Engine/Plugins/Developer/Concert/ConcertMain/Source/Concert/Public/IConcertServer.h
@@ -138,9 +138,10 @@ public:
 	 * @param ArchiveNameOverride The name override to give to the archived session.
 	 * @param SessionFilter The filter controlling which activities from the session should be archived.
 	 * @param OutFailureReason The reason the operation fails if the function returns false, undefined otherwise.
+	 * @param ArchiveSessionIdOverride The ID the archived session is supposed to have
 	 * @return The ID of the archived session on success, or an invalid GUID otherwise.
 	 */
-	virtual FGuid ArchiveSession(const FGuid& SessionId, const FString& ArchiveNameOverride, const FConcertSessionFilter& SessionFilter, FText& OutFailureReason) = 0;
+	virtual FGuid ArchiveSession(const FGuid& SessionId, const FString& ArchiveNameOverride, const FConcertSessionFilter& SessionFilter, FText& OutFailureReason, FGuid ArchiveSessionIdOverride = FGuid::NewGuid()) = 0;
 
 	/**
 	 * Copy the session data to a destination folder for external usage.
diff --git a/Engine/Plugins/Developer/Concert/ConcertSync/ConcertSyncCore/Source/ConcertSyncCore/Private/HistoryEdition/HistoryDeletion.cpp b/Engine/Plugins/Developer/Concert/ConcertSync/ConcertSyncCore/Source/ConcertSyncCore/Private/HistoryEdition/HistoryDeletion.cpp
new file mode 100644
--- /dev/null
+++ b/Engine/Plugins/Developer/Concert/ConcertSync/ConcertSyncCore/Source/ConcertSyncCore/Private/HistoryEdition/HistoryDeletion.cpp
@@ -0,0 +1,73 @@
+// Copyright Epic Games, Inc. All Rights Reserved.
+
+#include "HistoryEdition/HistoryDeletion.h"
+
+#include "IConcertServer.h"
+
+#include "Misc/ScopeExit.h"
+
+namespace UE::ConcertSyncCore
+{
+	static FConcertSessionFilter BuildFilterFrom(const TSet<FActivityID>& ToDelete);
+
+	TSet<FActivityID> CombineRequirements(const FHistoryDeletionRequirements& ToDelete)
+	{
+		TSet<FActivityID> Result;
+		for (const FActivityID ActivityID : ToDelete.HardDependencies)
+		{
+			Result.Add(ActivityID);
+		}
+		for (const FActivityID ActivityID : ToDelete.PossibleDependencies)
+		{
+			Result.Add(ActivityID);
+		}
+		return Result;
+	}
+	
+	FDeleteSessionErrorResult DeleteActivitiesInArchivedSession(const TSharedRef<IConcertServer>& Server, const FGuid& SessionToDeleteFrom, const TSet<FActivityID>& ToDelete)
+	{
+		const TOptional<FConcertSessionInfo> DeletedSessionInfo = Server->GetArchivedSessionInfo(SessionToDeleteFrom);
+		if (!DeletedSessionInfo)
+		{
+			const FString ErrorMessage = FString::Printf(TEXT("Session ID %s does not resolve to any archived session!"), *SessionToDeleteFrom.ToString());
+			return FDeleteSessionErrorResult::MakeError(FText::FromString(ErrorMessage));
+		}
+
+		// Restore the session while skipping all to be deleted activities
+		FText FailureReason;
+		const FConcertSessionFilter Filter = BuildFilterFrom(ToDelete);
+		const TSharedPtr<IConcertServerSession> LiveSession = Server->RestoreSession(SessionToDeleteFrom, *DeletedSessionInfo, Filter, FailureReason);
+		if (!LiveSession)
+		{
+			return FDeleteSessionErrorResult::MakeError(FailureReason);
+		}
+		ON_SCOPE_EXIT
+		{
+			Server->DestroySession(LiveSession->GetId(), FailureReason);
+		};
+
+		// The archived session must be removed before it can be overwritten
+		if (!Server->DestroySession(SessionToDeleteFrom, FailureReason))
+		{
+			return FDeleteSessionErrorResult::MakeError(FailureReason); 
+		}
+
+		if (!Server->ArchiveSession(LiveSession->GetId(), DeletedSessionInfo->SessionName, Filter, FailureReason, SessionToDeleteFrom).IsValid())
+		{
+			return FDeleteSessionErrorResult::MakeError(FailureReason); 
+		}
+		return FDeleteSessionErrorResult::MakeSuccess(); 
+	}
+
+	static FConcertSessionFilter BuildFilterFrom(const TSet<FActivityID>& ToDelete)
+	{
+		FConcertSessionFilter Result;
+		Result.ActivityIdsToExclude.Reserve(ToDelete.Num());
+		for (const FActivityID ActivityID : ToDelete)
+		{
+			Result.ActivityIdsToExclude.Add(ActivityID);
+		}
+		return Result;
+	}
+}
+
diff --git a/Engine/Plugins/Developer/Concert/ConcertSync/ConcertSyncCore/Source/ConcertSyncCore/Public/HistoryEdition/HistoryAnalysis.h b/Engine/Plugins/Developer/Concert/ConcertSync/ConcertSyncCore/Source/ConcertSyncCore/Public/HistoryEdition/HistoryAnalysis.h
--- a/Engine/Plugins/Developer/Concert/ConcertSync/ConcertSyncCore/Source/ConcertSyncCore/Public/HistoryEdition/HistoryAnalysis.h
+++ b/Engine/Plugins/Developer/Concert/ConcertSync/ConcertSyncCore/Source/ConcertSyncCore/Public/HistoryEdition/HistoryAnalysis.h
@@ -20,6 +20,11 @@ namespace UE::ConcertSyncCore
 		 * This will not contain any elements in HardDependencies.
 		 */
 		TSet<FActivityID> PossibleDependencies;
+
+		FHistoryDeletionRequirements(TSet<FActivityID> HardDependencies = {}, TSet<FActivityID> PossibleDependencies = {})
+			: HardDependencies(MoveTemp(HardDependencies)),
+			  PossibleDependencies(MoveTemp(PossibleDependencies))
+		{}
 	};
 
 	/** Utility function for one-off operations: just computes the dependency graph before calling AnalyseActivityDeletion. */
diff --git a/Engine/Plugins/Developer/Concert/ConcertSync/ConcertSyncCore/Source/ConcertSyncCore/Public/HistoryEdition/HistoryDeletion.h b/Engine/Plugins/Developer/Concert/ConcertSync/ConcertSyncCore/Source/ConcertSyncCore/Public/HistoryEdition/HistoryDeletion.h
new file mode 100644
--- /dev/null
+++ b/Engine/Plugins/Developer/Concert/ConcertSync/ConcertSyncCore/Source/ConcertSyncCore/Public/HistoryEdition/HistoryDeletion.h
@@ -0,0 +1,39 @@
+// Copyright Epic Games, Inc. All Rights Reserved.
+
+#pragma once
+
+#include "CoreMinimal.h"
+#include "HistoryAnalysis.h"
+
+class IConcertServer;
+
+namespace UE::ConcertSyncCore
+{
+	struct FDeleteSessionErrorResult
+	{
+		TOptional<FText> ErrorMessage;
+
+		explicit FDeleteSessionErrorResult(TOptional<FText> ErrorMessage = {})
+			: ErrorMessage(ErrorMessage)
+		{}
+
+		static FDeleteSessionErrorResult MakeSuccess() { return FDeleteSessionErrorResult(); }
+		static FDeleteSessionErrorResult MakeError(FText Error) { return FDeleteSessionErrorResult(MoveTemp(Error)); }
+
+		bool WasSuccessful() const { return !ErrorMessage.IsSet(); }
+		bool HadError() const { return !WasSuccessful(); }
+	};
+
+	/** Utility functions that converts FHistoryDeletionRequirements into a single TSet. */
+	CONCERTSYNCCORE_API TSet<FActivityID> CombineRequirements(const FHistoryDeletionRequirements& ToDelete);
+
+	/**
+	 * Deletes the given sessions in ToDelete from ArchivedSessionDatabase.
+	 *
+	 * There is no direct functionality for removing activities in the database; the operation is as follows:
+	 *	1. Restore the archived session with a session filter
+	 *	2. Delete the archived session
+	 *	3. Archive the live session created in step 1 with the old guid.
+	 */
+	CONCERTSYNCCORE_API FDeleteSessionErrorResult DeleteActivitiesInArchivedSession(const TSharedRef<IConcertServer>& Server, const FGuid& SessionToDeleteFrom, const TSet<FActivityID>& ToDelete);
+}
diff --git a/Engine/Plugins/Developer/Concert/ConcertSync/ConcertSyncTest/Source/ConcertSyncTest/Private/HistoryEdition/RenameEditAndDeleteMapsFlow.cpp b/Engine/Plugins/Developer/Concert/ConcertSync/ConcertSyncTest/Source/ConcertSyncTest/Private/HistoryEdition/RenameEditAndDeleteMapsFlow.cpp
--- a/Engine/Plugins/Developer/Concert/ConcertSync/ConcertSyncTest/Source/ConcertSyncTest/Private/HistoryEdition/RenameEditAndDeleteMapsFlow.cpp
+++ b/Engine/Plugins/Developer/Concert/ConcertSync/ConcertSyncTest/Source/ConcertSyncTest/Private/HistoryEdition/RenameEditAndDeleteMapsFlow.cpp
@@ -64,7 +64,7 @@ namespace UE::ConcertSyncTests::RenameEditAndDeleteMapsFlowTest
 	
 	FCreatedStaticMeshActor CreateEditedActor(FName OuterLevelPath);
 
-	TTestActivityArray<int64> CreateActivityHistory(FConcertSyncSessionDatabase& SessionDatabase, const FGuid& EndpointID)
+	TTestActivityArray<FActivityID> CreateActivityHistory(FConcertSyncSessionDatabase& SessionDatabase, const FGuid& EndpointID)
 	{
 		TTestActivityArray<int64> ActivityIDs;
 		TTestActivityArray<int64> PackageEventIDs;
diff --git a/Engine/Plugins/Developer/Concert/ConcertSync/ConcertSyncTest/Source/ConcertSyncTest/Private/HistoryEdition/RenameEditAndDeleteMapsFlow.h b/Engine/Plugins/Developer/Concert/ConcertSync/ConcertSyncTest/Source/ConcertSyncTest/Private/HistoryEdition/RenameEditAndDeleteMapsFlow.h
--- a/Engine/Plugins/Developer/Concert/ConcertSync/ConcertSyncTest/Source/ConcertSyncTest/Private/HistoryEdition/RenameEditAndDeleteMapsFlow.h
+++ b/Engine/Plugins/Developer/Concert/ConcertSync/ConcertSyncTest/Source/ConcertSyncTest/Private/HistoryEdition/RenameEditAndDeleteMapsFlow.h
@@ -3,6 +3,7 @@
 #pragma once
 
 #include "CoreMinimal.h"
+#include "ConcertMessageData.h"
 
 class FConcertSyncSessionDatabase;
 
@@ -44,5 +45,5 @@ namespace UE::ConcertSyncTests::RenameEditAndDeleteMapsFlowTest
 	 *
 	 *	@return An array where every ETestActivity entry is the index to the activity ID added to SessionDatabase
 	 */
-	TTestActivityArray<int64> CreateActivityHistory(FConcertSyncSessionDatabase& SessionDatabase, const FGuid& EndpointID);
+	TTestActivityArray<FActivityID> CreateActivityHistory(FConcertSyncSessionDatabase& SessionDatabase, const FGuid& EndpointID);
 }
\ No newline at end of file
