MU bug. Dynamic material instances do not transact to nDisplay until Save All

diff --git a/Engine/Plugins/Developer/Concert/ConcertSync/ConcertSyncCore/Config/BaseConcertSyncCore.ini b/Engine/Plugins/Developer/Concert/ConcertSync/ConcertSyncCore/Config/BaseConcertSyncCore.ini
--- a/Engine/Plugins/Developer/Concert/ConcertSync/ConcertSyncCore/Config/BaseConcertSyncCore.ini
+++ b/Engine/Plugins/Developer/Concert/ConcertSync/ConcertSyncCore/Config/BaseConcertSyncCore.ini
@@ -2,6 +2,7 @@
 +IncludeObjectClassFilters=(ObjectOuterClass=None,ObjectClasses=(Class'/Script/Engine.World',Class'/Script/Engine.BookmarkBase',Class'/Script/Engine.MaterialParameterCollection',Class'/Script/MovieScene.MovieSceneSignedObject',Class'/Script/MovieScene.MovieSceneFolder'))
 +IncludeObjectClassFilters=(ObjectOuterClass=Class'/Script/Engine.World',ObjectClasses=(Class'/Script/Engine.Layer',Class'/Script/Engine.Level',Class'/Script/Engine.LevelStreaming',Class'/Script/Engine.Actor',Class'/Script/Engine.ActorComponent',Class'/Script/Engine.Polys',Class'/Script/Engine.BrushBuilder',Class'/Script/Foliage.FoliageType'))
 +IncludeObjectClassFilters=(ObjectOuterClass=Class'/Script/MovieScene.MovieSceneSignedObject',ObjectClasses=(Class'/Script/Engine.Actor',Class'/Script/Engine.ActorComponent'))
++IncludeObjectClassFilters=(ObjectOuterClass=Class'/Script/Engine.World',ObjectClasses=(Class'/Script/Engine.MaterialInstanceDynamic'))
 +ExcludeTransactionClassFilters=(ObjectOuterClass=None,ObjectClasses=(Class'/Script/Landscape.LandscapeProxy',Class'/Script/Landscape.LandscapeHeightfieldCollisionComponent',Class'/Script/Landscape.LandscapeComponent'))
 +AllowedTransientProperties="/Script/Engine.Actor:bHiddenEdTemporary"
 +AllowedTransientProperties="/Script/Engine.Actor:bHiddenEdLevel"
diff --git a/Engine/Plugins/Experimental/VirtualProductionUtilities/Source/VPMaterialsEditor/Private/SMaterialDynamicWidgets.cpp b/Engine/Plugins/Experimental/VirtualProductionUtilities/Source/VPMaterialsEditor/Private/SMaterialDynamicWidgets.cpp
--- a/Engine/Plugins/Experimental/VirtualProductionUtilities/Source/VPMaterialsEditor/Private/SMaterialDynamicWidgets.cpp
+++ b/Engine/Plugins/Experimental/VirtualProductionUtilities/Source/VPMaterialsEditor/Private/SMaterialDynamicWidgets.cpp
@@ -143,7 +143,7 @@ FReply SMaterialDynamicView::OnRevertButtonClicked() const
 	}
 
 	MaterialItemView->ReplaceMaterial(ParentMaterialInstance);
-
+	
 	return FReply::Handled();
 }
 
@@ -187,15 +187,21 @@ FReply SMaterialDynamicView::OnCreateDynamicMaterialButtonClicked() const
 	}
 	
 	UMaterialInstance* MaterialInstance = Cast<UMaterialInstance>(MaterialItemView->GetMaterialListItem().Material.Get());
+	UActorComponent* OwnerActorComponent = CurrentComponent.Get();
 
-	if (!CurrentComponent.IsValid())
+	if (!OwnerActorComponent)
 	{
 		return FReply::Handled();
 	}
 
 	if (MaterialInstance)
 	{
-		UMaterialInstanceDynamic* NewMaterialInstanceDynamic = UMaterialInstanceDynamic::Create(MaterialInstance, CurrentComponent.Get());
+		// Create Material
+		UMaterialInstanceDynamic* NewMaterialInstanceDynamic = UMaterialInstanceDynamic::Create(MaterialInstance, OwnerActorComponent);
+
+		// Set object Transactional in order Undo, Redo working correctly and have MU support
+		NewMaterialInstanceDynamic->SetFlags(RF_Transactional);
+		
 		NewMaterialInstanceDynamic->CopyParameterOverrides(MaterialInstance);
 		MaterialItemView->ReplaceMaterial(NewMaterialInstanceDynamic);
 	}
