Allow override of concert function to determine if object path belongs to a world

diff --git a/Engine/Plugins/Developer/Concert/ConcertSync/ConcertSyncCore/Source/ConcertSyncCore/Private/ConcertSyncArchives.cpp b/Engine/Plugins/Developer/Concert/ConcertSync/ConcertSyncCore/Source/ConcertSyncCore/Private/ConcertSyncArchives.cpp
--- a/Engine/Plugins/Developer/Concert/ConcertSync/ConcertSyncCore/Source/ConcertSyncCore/Private/ConcertSyncArchives.cpp
+++ b/Engine/Plugins/Developer/Concert/ConcertSync/ConcertSyncCore/Source/ConcertSyncCore/Private/ConcertSyncArchives.cpp
@@ -60,6 +60,11 @@ FString FConcertSyncWorldRemapper::RemapObjectPathName(const FString& InObjectPa
 
 bool FConcertSyncWorldRemapper::ObjectBelongsToWorld(const FString& InObjectPathName) const
 {
+	if (ObjectPathBelongsToWorldDelegate.IsBound() && ObjectPathBelongsToWorldDelegate.Execute(FStringView(InObjectPathName)))
+	{
+		return true;
+	}
+
 	return HasMapping() && (InObjectPathName.StartsWith(SourceWorldPathName) || InObjectPathName.StartsWith(DestWorldPathName));
 }
 
@@ -268,20 +273,23 @@ FArchive& FConcertSyncObjectReader::operator<<(UObject*& Obj)
 		// Always attempt to find an in-memory object first as we may be calling this function while a load is taking place
 		Obj = StaticFindObject(UObject::StaticClass(), nullptr, *ResolvedObjPath);
 
-		// We do not attempt to load objects within the current world as they may not have been created yet, 
-		// and we don't want to trigger a reload of the world package (when iterative cooking is enabled)
-		const bool bAllowLoad = !WorldRemapper.ObjectBelongsToWorld(ResolvedObjPath);
-		if (!Obj && bAllowLoad)
+		if (!Obj)
 		{
-			// If the outer name is a package path that isn't currently loaded, then we need to try loading it to avoid 
-			// creating an in-memory version of the package (which would prevent the real package ever loading)
-			if (FPackageName::IsValidLongPackageName(ResolvedObjPath))
-			{
-				Obj = LoadPackage(nullptr, *ResolvedObjPath, LOAD_NoWarn);
-			}
-			else
+			// We do not attempt to load objects within the current world as they may not have been created yet, 
+			// and we don't want to trigger a reload of the world package (when iterative cooking is enabled)
+			const bool bAllowLoad = !WorldRemapper.ObjectBelongsToWorld(ResolvedObjPath);
+			if (bAllowLoad)
 			{
-				Obj = StaticLoadObject(UObject::StaticClass(), nullptr, *ResolvedObjPath);
+				// If the outer name is a package path that isn't currently loaded, then we need to try loading it to avoid 
+				// creating an in-memory version of the package (which would prevent the real package ever loading)
+				if (FPackageName::IsValidLongPackageName(ResolvedObjPath))
+				{
+					Obj = LoadPackage(nullptr, *ResolvedObjPath, LOAD_NoWarn);
+				}
+				else
+				{
+					Obj = StaticLoadObject(UObject::StaticClass(), nullptr, *ResolvedObjPath);
+				}
 			}
 		}
 	}
diff --git a/Engine/Plugins/Developer/Concert/ConcertSync/ConcertSyncCore/Source/ConcertSyncCore/Public/ConcertSyncArchives.h b/Engine/Plugins/Developer/Concert/ConcertSync/ConcertSyncCore/Source/ConcertSyncCore/Public/ConcertSyncArchives.h
--- a/Engine/Plugins/Developer/Concert/ConcertSync/ConcertSyncCore/Source/ConcertSyncCore/Public/ConcertSyncArchives.h
+++ b/Engine/Plugins/Developer/Concert/ConcertSync/ConcertSyncCore/Source/ConcertSyncCore/Public/ConcertSyncArchives.h
@@ -8,6 +8,7 @@
 struct FConcertSessionVersionInfo;
 
 DECLARE_DELEGATE_OneParam(FConcertSyncRemapObjectPath, FString& /*ObjectPath*/)
+DECLARE_DELEGATE_RetVal_OneParam(bool, FConcertSyncObjectPathBelongsToWorld, const FStringView /*ObjectPath*/)
 
 namespace ConcertSyncUtil
 {
@@ -33,6 +34,7 @@ public:
 	bool HasMapping() const;
 
 	FConcertSyncRemapObjectPath RemapDelegate;
+	FConcertSyncObjectPathBelongsToWorld ObjectPathBelongsToWorldDelegate;
 
 private:
 	FString SourceWorldPathName;
