Fix crash when encountering a Dummy transaction

diff --git a/Engine/Plugins/Developer/Concert/ConcertSync/ConcertSyncCore/Source/ConcertSyncCore/Private/HistoryEdition/DependencyGraphBuilder.cpp b/Engine/Plugins/Developer/Concert/ConcertSync/ConcertSyncCore/Source/ConcertSyncCore/Private/HistoryEdition/DependencyGraphBuilder.cpp
--- a/Engine/Plugins/Developer/Concert/ConcertSync/ConcertSyncCore/Source/ConcertSyncCore/Private/HistoryEdition/DependencyGraphBuilder.cpp
+++ b/Engine/Plugins/Developer/Concert/ConcertSync/ConcertSyncCore/Source/ConcertSyncCore/Private/HistoryEdition/DependencyGraphBuilder.cpp
@@ -206,7 +206,9 @@ namespace UE::ConcertSyncCore
 			DiscoverDeletedPackageDependencies(NodeID, Activity, EventData);
 			break;
 			
-		case EConcertPackageUpdateType::Dummy: 
+		case EConcertPackageUpdateType::Dummy:
+			// No idea what dependencies this has ... let's assume none and hopefully not regret it later
+			break;
 		default:
 			checkNoEntry();
 		}
