Add flags to FConcertSyncActivity together with ability to mute activities. Muted activities are not sent to the client.

diff --git a/Engine/Plugins/Developer/Concert/ConcertMain/Source/Concert/Private/ConcertServer.cpp b/Engine/Plugins/Developer/Concert/ConcertMain/Source/Concert/Private/ConcertServer.cpp
--- a/Engine/Plugins/Developer/Concert/ConcertMain/Source/Concert/Private/ConcertServer.cpp
+++ b/Engine/Plugins/Developer/Concert/ConcertMain/Source/Concert/Private/ConcertServer.cpp
@@ -1471,7 +1471,7 @@ TFuture<FConcertAdmin_GetSessionActivitiesResponse> FConcertServer::HandleGetSes
 	FConcertAdmin_GetSessionActivitiesResponse ResponseData;
 
 	const FConcertAdmin_GetSessionActivitiesRequest* Message = Context.GetMessage<FConcertAdmin_GetSessionActivitiesRequest>();
-	if (EventSink->GetSessionActivities(*this, Message->SessionId, Message->FromActivityId, Message->ActivityCount, ResponseData.Activities, ResponseData.EndpointClientInfoMap, Message->bIncludeDetails))
+	if (EventSink->GetUnmutedSessionActivities(*this, Message->SessionId, Message->FromActivityId, Message->ActivityCount, ResponseData.Activities, ResponseData.EndpointClientInfoMap, Message->bIncludeDetails))
 	{
 		ResponseData.ResponseCode = EConcertResponseCode::Success;
 	}
diff --git a/Engine/Plugins/Developer/Concert/ConcertMain/Source/Concert/Public/IConcertServerEventSink.h b/Engine/Plugins/Developer/Concert/ConcertMain/Source/Concert/Public/IConcertServerEventSink.h
--- a/Engine/Plugins/Developer/Concert/ConcertMain/Source/Concert/Public/IConcertServerEventSink.h
+++ b/Engine/Plugins/Developer/Concert/ConcertMain/Source/Concert/Public/IConcertServerEventSink.h
@@ -84,10 +84,10 @@ public:
 	virtual bool RestoreSession(const IConcertServer& InServer, const FGuid& InArchivedSessionId, const FString& InLiveSessionRoot, const FConcertSessionInfo& InLiveSessionInfo, const FConcertSessionFilter& InSessionFilter) = 0;
 
 	/**
-	 * Called to get the activities for an archived or a live session without being connected to it.
+	 * Called to get the all activities that are not muted (i.e. don't have EConcertSyncActivityFlags::Muted flag set) for an archived or a live session without being connected to it.
 	 * @note If ActivityCount is negative, the function returns the last activities (the tail) from Max(1, TotalActivityCount + ActivityCount + 1)
 	 */
-	virtual bool GetSessionActivities(const IConcertServer& InServer, const FGuid& SessionId, int64 FromActivityId, int64 ActivityCount, TArray<FConcertSessionSerializedPayload>& OutActivities, TMap<FGuid, FConcertClientInfo>& OutEndpointClientInfoMap, bool bIncludeDetails) = 0;
+	virtual bool GetUnmutedSessionActivities(const IConcertServer& InServer, const FGuid& SessionId, int64 FromActivityId, int64 ActivityCount, TArray<FConcertSessionSerializedPayload>& OutActivities, TMap<FGuid, FConcertClientInfo>& OutEndpointClientInfoMap, bool bIncludeDetails) = 0;
 
 	/**
 	 * Called when a live session is renamed.
diff --git a/Engine/Plugins/Developer/Concert/ConcertSync/ConcertSyncClient/Source/ConcertSyncClient/Private/ConcertClientWorkspace.cpp b/Engine/Plugins/Developer/Concert/ConcertSync/ConcertSyncClient/Source/ConcertSyncClient/Private/ConcertClientWorkspace.cpp
--- a/Engine/Plugins/Developer/Concert/ConcertSync/ConcertSyncClient/Source/ConcertSyncClient/Private/ConcertClientWorkspace.cpp
+++ b/Engine/Plugins/Developer/Concert/ConcertSync/ConcertSyncClient/Source/ConcertSyncClient/Private/ConcertClientWorkspace.cpp
@@ -910,6 +910,7 @@ void FConcertClientWorkspace::HandleWorkspaceSyncActivityEvent(const FConcertSes
 
 	check(ActivityPayload.IsValid() && ActivityPayload.GetStruct()->IsChildOf(FConcertSyncActivity::StaticStruct()));
 	FConcertSyncActivity* Activity = (FConcertSyncActivity*)ActivityPayload.GetStructMemory();
+	ensureAlwaysMsgf((Activity->Flags & EConcertSyncActivityFlags::Muted) == EConcertSyncActivityFlags::None, TEXT("Clients are not supposed to receive muted activities!"));
 
 	// Update slow task dialog
 	if (InitialSyncSlowTask.IsValid())
diff --git a/Engine/Plugins/Developer/Concert/ConcertSync/ConcertSyncCore/Source/ConcertSyncCore/Private/ConcertSyncSessionDatabase.cpp b/Engine/Plugins/Developer/Concert/ConcertSync/ConcertSyncCore/Source/ConcertSyncCore/Private/ConcertSyncSessionDatabase.cpp
--- a/Engine/Plugins/Developer/Concert/ConcertSync/ConcertSyncCore/Source/ConcertSyncCore/Private/ConcertSyncSessionDatabase.cpp
+++ b/Engine/Plugins/Developer/Concert/ConcertSync/ConcertSyncCore/Source/ConcertSyncCore/Private/ConcertSyncSessionDatabase.cpp
@@ -544,7 +544,9 @@ enum class FConcertSyncSessionDatabaseVersion
 	Empty = 0,
 	Initial = 1,
 	V2 = 2,
-	Current = V2,
+	/** Added FConcertSyncActivity::Flags to the activities table */
+	AddActivityFlags = 3,
+	Current = AddActivityFlags,
 };
 
 class FConcertSyncSessionDatabaseStatements
@@ -618,6 +620,7 @@ public:
 		PREPARE_STATEMENT(Statement_GetActivityDataInRange);
 		PREPARE_STATEMENT(Statement_GetAllActivityIdAndEventTypes);
 		PREPARE_STATEMENT(Statement_GetActivityIdAndEventTypesInRange);
+		PREPARE_STATEMENT(Statement_GetActivityIdWithEventTypesAndFlagsInRange);
 		PREPARE_STATEMENT(Statement_GetActivityMaxId);
 		
 		PREPARE_STATEMENT(Statement_IgnoreActivity);
@@ -1065,11 +1068,11 @@ public:
 	 */
 
 	/** Add the activity data to activities and get its activity_id */
-	SQLITE_PREPARED_STATEMENT_BINDINGS_ONLY(FAddActivityData, "INSERT INTO activities(endpoint_id, event_time, event_type, event_id, event_summary_type, event_summary_flags, event_summary_size_bytes, event_summary_data) VALUES(?1, ?2, ?3, ?4, ?5, ?6, ?7, ?8);", SQLITE_PREPARED_STATEMENT_BINDINGS(FGuid, FDateTime, EConcertSyncActivityEventType, int64, FName, int32, int32, TArray<uint8>));
+	SQLITE_PREPARED_STATEMENT_BINDINGS_ONLY(FAddActivityData, "INSERT INTO activities(activity_flags, endpoint_id, event_time, event_type, event_id, event_summary_type, event_summary_flags, event_summary_size_bytes, event_summary_data) VALUES(?1, ?2, ?3, ?4, ?5, ?6, ?7, ?8, ?9);", SQLITE_PREPARED_STATEMENT_BINDINGS(EConcertSyncActivityFlags, FGuid, FDateTime, EConcertSyncActivityEventType, int64, FName, int32, int32, TArray<uint8>));
 	FAddActivityData Statement_AddActivityData;
-	bool AddActivityData(const FGuid& InEndpointId, const EConcertSyncActivityEventType InEventType, const int64 InEventId, const FConcertSessionSerializedPayload& InEventSummary, int64& OutActivityId)
+	bool AddActivityData(const EConcertSyncActivityFlags Flags, const FGuid& InEndpointId, const EConcertSyncActivityEventType InEventType, const int64 InEventId, const FConcertSessionSerializedPayload& InEventSummary, int64& OutActivityId)
 	{
-		if (Statement_AddActivityData.BindAndExecute(InEndpointId, FDateTime::UtcNow(), InEventType,
+		if (Statement_AddActivityData.BindAndExecute(Flags, InEndpointId, FDateTime::UtcNow(), InEventType,
 													 InEventId, InEventSummary.PayloadTypeName,
 													 PackageDataUtil::ConvertSerializedPayloadFlagsToInt32(InEventSummary),
 													 InEventSummary.PayloadSize,
@@ -1082,11 +1085,11 @@ public:
 	}
 
 	/** Set the activity data in activities for the given activity_id */
-	SQLITE_PREPARED_STATEMENT_BINDINGS_ONLY(FSetActivityData, "INSERT OR REPLACE INTO activities(activity_id, endpoint_id, event_time, event_type, event_id, event_summary_type, event_summary_flags, event_summary_size_bytes, event_summary_data) VALUES(?1, ?2, ?3, ?4, ?5, ?6, ?7, ?8, ?9);", SQLITE_PREPARED_STATEMENT_BINDINGS(int64, FGuid, FDateTime, EConcertSyncActivityEventType, int64, FName, int32, int32, TArray<uint8>));
+	SQLITE_PREPARED_STATEMENT_BINDINGS_ONLY(FSetActivityData, "INSERT OR REPLACE INTO activities(activity_id, activity_flags, endpoint_id, event_time, event_type, event_id, event_summary_type, event_summary_flags, event_summary_size_bytes, event_summary_data) VALUES(?1, ?2, ?3, ?4, ?5, ?6, ?7, ?8, ?9, ?10);", SQLITE_PREPARED_STATEMENT_BINDINGS(int64, EConcertSyncActivityFlags, FGuid, FDateTime, EConcertSyncActivityEventType, int64, FName, int32, int32, TArray<uint8>));
 	FSetActivityData Statement_SetActivityData;
-	bool SetActivityData(const int64 InActivityId, const FGuid& InEndpointId, const FDateTime InEventTime, const EConcertSyncActivityEventType InEventType, const int64 InEventId, const FConcertSessionSerializedPayload& InEventSummary)
+	bool SetActivityData(const int64 InActivityId, const EConcertSyncActivityFlags Flags, const FGuid& InEndpointId, const FDateTime InEventTime, const EConcertSyncActivityEventType InEventType, const int64 InEventId, const FConcertSessionSerializedPayload& InEventSummary)
 	{
-		return Statement_SetActivityData.BindAndExecute(InActivityId, InEndpointId, InEventTime, InEventType, InEventId,
+		return Statement_SetActivityData.BindAndExecute(InActivityId, Flags, InEndpointId, InEventTime, InEventType, InEventId,
 														InEventSummary.PayloadTypeName,
 														PackageDataUtil::ConvertSerializedPayloadFlagsToInt32(InEventSummary),
 														InEventSummary.PayloadSize,
@@ -1094,12 +1097,12 @@ public:
 	}
 
 	/** Get the activity data from activities for the given activity_id */
-	SQLITE_PREPARED_STATEMENT(FGetActivityDataForId, "SELECT endpoint_id, event_time, event_type, event_id, event_summary_type, event_summary_flags, event_summary_size_bytes, event_summary_data FROM activities WHERE activity_id = ?1;", SQLITE_PREPARED_STATEMENT_COLUMNS(FGuid, FDateTime, EConcertSyncActivityEventType, int64, FName, int32, int32, TArray<uint8>), SQLITE_PREPARED_STATEMENT_BINDINGS(int64));
+	SQLITE_PREPARED_STATEMENT(FGetActivityDataForId, "SELECT activity_flags, endpoint_id, event_time, event_type, event_id, event_summary_type, event_summary_flags, event_summary_size_bytes, event_summary_data FROM activities WHERE activity_id = ?1;", SQLITE_PREPARED_STATEMENT_COLUMNS(EConcertSyncActivityFlags, FGuid, FDateTime, EConcertSyncActivityEventType, int64, FName, int32, int32, TArray<uint8>), SQLITE_PREPARED_STATEMENT_BINDINGS(int64));
 	FGetActivityDataForId Statement_GetActivityDataForId;
-	bool GetActivityDataForId(const int64 InActivityId, FGuid& OutEndpointId, FDateTime& OutEventTime, EConcertSyncActivityEventType& OutEventType, int64& OutEventId, FConcertSessionSerializedPayload& OutEventSummary)
+	bool GetActivityDataForId(const int64 InActivityId, EConcertSyncActivityFlags& OutFlags, FGuid& OutEndpointId, FDateTime& OutEventTime, EConcertSyncActivityEventType& OutEventType, int64& OutEventId, FConcertSessionSerializedPayload& OutEventSummary)
 	{
 		int32 PayloadFlags;
-		if(  Statement_GetActivityDataForId.BindAndExecuteSingle(InActivityId, OutEndpointId, OutEventTime, OutEventType, OutEventId,
+		if(  Statement_GetActivityDataForId.BindAndExecuteSingle(InActivityId, OutFlags, OutEndpointId, OutEventTime, OutEventType, OutEventId,
 																 OutEventSummary.PayloadTypeName,
 																 PayloadFlags,
 																 OutEventSummary.PayloadSize,
@@ -1112,12 +1115,12 @@ public:
 	}
 
 	/** Get the activity data from activities for the given event_id and event_type */
-	SQLITE_PREPARED_STATEMENT(FGetActivityDataForEvent, "SELECT activity_id, endpoint_id, event_time, event_summary_type, event_summary_flags, event_summary_size_bytes, event_summary_data FROM activities WHERE event_id = ?1 AND event_type = ?2;", SQLITE_PREPARED_STATEMENT_COLUMNS(int64, FGuid, FDateTime, FName, int32, int32, TArray<uint8>), SQLITE_PREPARED_STATEMENT_BINDINGS(int64, EConcertSyncActivityEventType));
+	SQLITE_PREPARED_STATEMENT(FGetActivityDataForEvent, "SELECT activity_id, activity_flags, endpoint_id, event_time, event_summary_type, event_summary_flags, event_summary_size_bytes, event_summary_data FROM activities WHERE event_id = ?1 AND event_type = ?2;", SQLITE_PREPARED_STATEMENT_COLUMNS(int64, EConcertSyncActivityFlags, FGuid, FDateTime, FName, int32, int32, TArray<uint8>), SQLITE_PREPARED_STATEMENT_BINDINGS(int64, EConcertSyncActivityEventType));
 	FGetActivityDataForEvent Statement_GetActivityDataForEvent;
-	bool GetActivityDataForEvent(const int64 InEventId, const EConcertSyncActivityEventType InEventType, int64& OutActivityId, FGuid& OutEndpointId, FDateTime& OutEventTime, FConcertSessionSerializedPayload& OutEventSummary)
+	bool GetActivityDataForEvent(const int64 InEventId, const EConcertSyncActivityEventType InEventType, int64& OutActivityId, EConcertSyncActivityFlags& OutFlags, FGuid& OutEndpointId, FDateTime& OutEventTime, FConcertSessionSerializedPayload& OutEventSummary)
 	{
 		int32 PayloadFlags;
-		if (  Statement_GetActivityDataForEvent.BindAndExecuteSingle(InEventId, InEventType, OutActivityId, OutEndpointId,
+		if (  Statement_GetActivityDataForEvent.BindAndExecuteSingle(InEventId, InEventType, OutActivityId, OutFlags, OutEndpointId,
 																	 OutEventTime,
 																	 OutEventSummary.PayloadTypeName,
 																	 PayloadFlags,
@@ -1139,80 +1142,83 @@ public:
 	}
 
 	/** Get the activity data from activities for all activity_ids */
-	SQLITE_PREPARED_STATEMENT_COLUMNS_ONLY(FGetAllActivityData, "SELECT activity_id, endpoint_id, event_time, event_type, event_id, event_summary_type, event_summary_flags, event_summary_size_bytes, event_summary_data FROM activities ORDER BY activity_id;", SQLITE_PREPARED_STATEMENT_COLUMNS(int64, FGuid, FDateTime, EConcertSyncActivityEventType, int64, FName, int32, int32, TArray<uint8>));
+	SQLITE_PREPARED_STATEMENT_COLUMNS_ONLY(FGetAllActivityData, "SELECT activity_id, activity_flags, endpoint_id, event_time, event_type, event_id, event_summary_type, event_summary_flags, event_summary_size_bytes, event_summary_data FROM activities ORDER BY activity_id;", SQLITE_PREPARED_STATEMENT_COLUMNS(int64, EConcertSyncActivityFlags, FGuid, FDateTime, EConcertSyncActivityEventType, int64, FName, int32, int32, TArray<uint8>));
 	FGetAllActivityData Statement_GetAllActivityData;
-	bool GetAllActivityData(TFunctionRef<ESQLitePreparedStatementExecuteRowResult(int64, const FGuid&, FDateTime, EConcertSyncActivityEventType, int64, FConcertSessionSerializedPayload&&)> InCallback)
+	bool GetAllActivityData(TFunctionRef<ESQLitePreparedStatementExecuteRowResult(int64, const EConcertSyncActivityFlags, const FGuid&, FDateTime, EConcertSyncActivityEventType, int64, FConcertSessionSerializedPayload&&)> InCallback)
 	{
 		return Statement_GetAllActivityData.Execute([&InCallback](const FGetAllActivityData& InStatement)
 		{
 			int64 ActivityId = 0;
+			EConcertSyncActivityFlags Flags = EConcertSyncActivityFlags::None;
 			FGuid EndpointId;
 			FDateTime EventTime;
 			EConcertSyncActivityEventType EventType = EConcertSyncActivityEventType::Connection;
 			int64 EventId = 0;
 			FConcertSessionSerializedPayload EventSummary(EConcertPayloadSerializationMethod::Cbor);
 			int32 PayloadFlags;
-			if (InStatement.GetColumnValues(ActivityId, EndpointId, EventTime, EventType, EventId,
+			if (InStatement.GetColumnValues(ActivityId, Flags, EndpointId, EventTime, EventType, EventId,
 											EventSummary.PayloadTypeName,
 											PayloadFlags,
 											EventSummary.PayloadSize,
 											EventSummary.PayloadBytes.Bytes))
 			{
 				PackageDataUtil::SetSerializedPayloadFlags(EventSummary, PayloadFlags);
-				return InCallback(ActivityId, EndpointId, EventTime, EventType, EventId, MoveTemp(EventSummary));
+				return InCallback(ActivityId, Flags, EndpointId, EventTime, EventType, EventId, MoveTemp(EventSummary));
 			}
 			return ESQLitePreparedStatementExecuteRowResult::Error;
 		}) != INDEX_NONE;
 	}
 
 	/** Get the activity data from activities for all activities of event_type */
-	SQLITE_PREPARED_STATEMENT(FGetAllActivityDataForEventType, "SELECT activity_id, endpoint_id, event_time, event_id, event_summary_type, event_summary_flags, event_summary_size_bytes, event_summary_data FROM activities WHERE event_type = ?1 ORDER BY activity_id;", SQLITE_PREPARED_STATEMENT_COLUMNS(int64, FGuid, FDateTime, int64, FName, int32, int32, TArray<uint8>), SQLITE_PREPARED_STATEMENT_BINDINGS(EConcertSyncActivityEventType));
+	SQLITE_PREPARED_STATEMENT(FGetAllActivityDataForEventType, "SELECT activity_id, activity_flags, endpoint_id, event_time, event_id, event_summary_type, event_summary_flags, event_summary_size_bytes, event_summary_data FROM activities WHERE event_type = ?1 ORDER BY activity_id;", SQLITE_PREPARED_STATEMENT_COLUMNS(int64, EConcertSyncActivityFlags, FGuid, FDateTime, int64, FName, int32, int32, TArray<uint8>), SQLITE_PREPARED_STATEMENT_BINDINGS(EConcertSyncActivityEventType));
 	FGetAllActivityDataForEventType Statement_GetAllActivityDataForEventType;
-	bool GetAllActivityDataForEventType(const EConcertSyncActivityEventType InEventType, TFunctionRef<ESQLitePreparedStatementExecuteRowResult(int64, const FGuid&, FDateTime, int64, FConcertSessionSerializedPayload&&)> InCallback)
+	bool GetAllActivityDataForEventType(const EConcertSyncActivityEventType InEventType, TFunctionRef<ESQLitePreparedStatementExecuteRowResult(int64, const EConcertSyncActivityFlags, const FGuid&, FDateTime, int64, FConcertSessionSerializedPayload&&)> InCallback)
 	{
 		return Statement_GetAllActivityDataForEventType.BindAndExecute(InEventType, [&InCallback](const FGetAllActivityDataForEventType& InStatement)
 		{
 			int64 ActivityId = 0;
+			EConcertSyncActivityFlags Flags = EConcertSyncActivityFlags::None;
 			FGuid EndpointId;
 			FDateTime EventTime;
 			int64 EventId = 0;
 			int32 PayloadFlags;
 			FConcertSessionSerializedPayload EventSummary(EConcertPayloadSerializationMethod::Cbor);
-			if (InStatement.GetColumnValues(ActivityId, EndpointId, EventTime, EventId,
+			if (InStatement.GetColumnValues(ActivityId, Flags, EndpointId, EventTime, EventId,
 											EventSummary.PayloadTypeName,
 											PayloadFlags,
 											EventSummary.PayloadSize,
 											EventSummary.PayloadBytes.Bytes))
 			{
 				PackageDataUtil::SetSerializedPayloadFlags(EventSummary, PayloadFlags);
-				return InCallback(ActivityId, EndpointId, EventTime, EventId, MoveTemp(EventSummary));
+				return InCallback(ActivityId, Flags, EndpointId, EventTime, EventId, MoveTemp(EventSummary));
 			}
 			return ESQLitePreparedStatementExecuteRowResult::Error;
 		}) != INDEX_NONE;
 	}
 
 	/** Get the activity data from activities for all activities in the given range */
-	SQLITE_PREPARED_STATEMENT(FGetActivityDataInRange, "SELECT activity_id, endpoint_id, event_time, event_type, event_id, event_summary_type, event_summary_flags, event_summary_size_bytes, event_summary_data FROM activities WHERE activity_id >= ?1 ORDER BY activity_id LIMIT ?2;", SQLITE_PREPARED_STATEMENT_COLUMNS(int64, FGuid, FDateTime, EConcertSyncActivityEventType, int64, FName, int32, int32, TArray<uint8>), SQLITE_PREPARED_STATEMENT_BINDINGS(int64, int64));
+	SQLITE_PREPARED_STATEMENT(FGetActivityDataInRange, "SELECT activity_id, activity_flags, endpoint_id, event_time, event_type, event_id, event_summary_type, event_summary_flags, event_summary_size_bytes, event_summary_data FROM activities WHERE activity_id >= ?1 ORDER BY activity_id LIMIT ?2;", SQLITE_PREPARED_STATEMENT_COLUMNS(int64, EConcertSyncActivityFlags, FGuid, FDateTime, EConcertSyncActivityEventType, int64, FName, int32, int32, TArray<uint8>), SQLITE_PREPARED_STATEMENT_BINDINGS(int64, int64));
 	FGetActivityDataInRange Statement_GetActivityDataInRange;
-	bool GetActivityDataInRange(const int64 InFirstActivityId, const int64 InMaxNumActivities, TFunctionRef<ESQLitePreparedStatementExecuteRowResult(int64, const FGuid&, FDateTime, EConcertSyncActivityEventType, int64, FConcertSessionSerializedPayload&&)> InCallback)
+	bool GetActivityDataInRange(const int64 InFirstActivityId, const int64 InMaxNumActivities, TFunctionRef<ESQLitePreparedStatementExecuteRowResult(int64, const EConcertSyncActivityFlags, const FGuid&, FDateTime, EConcertSyncActivityEventType, int64, FConcertSessionSerializedPayload&&)> InCallback)
 	{
 		return Statement_GetActivityDataInRange.BindAndExecute(InFirstActivityId, InMaxNumActivities, [&InCallback](const FGetActivityDataInRange& InStatement)
 		{
 			int64 ActivityId = 0;
+			EConcertSyncActivityFlags Flags = EConcertSyncActivityFlags::None;
 			FGuid EndpointId;
 			FDateTime EventTime;
 			EConcertSyncActivityEventType EventType = EConcertSyncActivityEventType::Connection;
 			int64 EventId = 0;
 			FConcertSessionSerializedPayload EventSummary(EConcertPayloadSerializationMethod::Cbor);
 			int32 EventFlags;
-			if (InStatement.GetColumnValues(ActivityId, EndpointId, EventTime, EventType, EventId,
+			if (InStatement.GetColumnValues(ActivityId, Flags, EndpointId, EventTime, EventType, EventId,
 											EventSummary.PayloadTypeName,
 											EventFlags,
 											EventSummary.PayloadSize,
 											EventSummary.PayloadBytes.Bytes))
 			{
 				PackageDataUtil::SetSerializedPayloadFlags(EventSummary, EventFlags);
-				return InCallback(ActivityId, EndpointId, EventTime, EventType, EventId, MoveTemp(EventSummary));
+				return InCallback(ActivityId, Flags, EndpointId, EventTime, EventType, EventId, MoveTemp(EventSummary));
 			}
 			return ESQLitePreparedStatementExecuteRowResult::Error;
 		}) != INDEX_NONE;
@@ -1252,6 +1258,24 @@ public:
 		}) != INDEX_NONE;
 	}
 
+	/** Get the activity_id, event_type and activity_flags from activities for all activities in the given range */
+	SQLITE_PREPARED_STATEMENT(FGetActivityIdWithEventTypesAndFlagsInRange, "SELECT activity_id, event_type, activity_flags FROM activities WHERE activity_id >= ?1 ORDER BY activity_id LIMIT ?2;", SQLITE_PREPARED_STATEMENT_COLUMNS(int64, EConcertSyncActivityEventType, EConcertSyncActivityFlags), SQLITE_PREPARED_STATEMENT_BINDINGS(int64, int64));
+	FGetActivityIdWithEventTypesAndFlagsInRange Statement_GetActivityIdWithEventTypesAndFlagsInRange;
+	bool GetActivityIdWithEventTypesAndFlagsInRange(const int64 InFirstActivityId, const int64 InMaxNumActivities, TFunctionRef<ESQLitePreparedStatementExecuteRowResult(int64, EConcertSyncActivityEventType, EConcertSyncActivityFlags)> InCallback)
+	{
+		return Statement_GetActivityIdWithEventTypesAndFlagsInRange.BindAndExecute(InFirstActivityId, InMaxNumActivities, [&InCallback](const FGetActivityIdWithEventTypesAndFlagsInRange& InStatement)
+		{
+			int64 ActivityId = 0;
+			EConcertSyncActivityEventType EventType = EConcertSyncActivityEventType::Connection;
+			EConcertSyncActivityFlags Flags = EConcertSyncActivityFlags::None;
+			if (InStatement.GetColumnValues(ActivityId, EventType, Flags))
+			{
+				return InCallback(ActivityId, EventType, Flags);
+			}
+			return ESQLitePreparedStatementExecuteRowResult::Error;
+		}) != INDEX_NONE;
+	}
+
 	/** Get the largest activity_id currently in activities */
 	SQLITE_PREPARED_STATEMENT_COLUMNS_ONLY(FGetActivityMaxId, "SELECT MAX(activity_id) FROM activities;", SQLITE_PREPARED_STATEMENT_COLUMNS(int64));
 	FGetActivityMaxId Statement_GetActivityMaxId;
@@ -1584,6 +1608,9 @@ bool FConcertSyncSessionDatabase::Open(const FString& InSessionPath, const ESQLi
 		return false;
 	}
 
+	const int64 NumberColumns = Database->Execute(TEXT("SELECT name FROM sqlite_master WHERE type='table' AND name='activities';"), [](const auto&){ return ESQLitePreparedStatementExecuteRowResult::Stop; });
+	const bool bActivitiesTableExisted = 1 == NumberColumns;
+
 	// Create our required tables
 #define CREATE_TABLE(NAME, STATEMENT)																										\
 	if (!Database->Execute(TEXT("CREATE TABLE IF NOT EXISTS ") TEXT(NAME) TEXT("(") TEXT(STATEMENT) TEXT(");")))							\
@@ -1599,13 +1626,23 @@ bool FConcertSyncSessionDatabase::Open(const FString& InSessionPath, const ESQLi
 	CREATE_TABLE("transaction_events", "transaction_event_id INTEGER PRIMARY KEY, data_filename TEXT NOT NULL");
 	CREATE_TABLE("package_events", "package_event_id INTEGER PRIMARY KEY, package_name_id INTEGER NOT NULL, package_revision INTEGER NOT NULL, package_info_flags INTEGER NOT NULL, package_info_size_bytes INTEGER NOT NULL, package_info_data BLOB, transaction_event_id_at_save INTEGER NOT NULL, data_filename TEXT NOT NULL, FOREIGN KEY(package_name_id) REFERENCES package_names(package_name_id)");
 	CREATE_TABLE("persist_events", "persist_event_id INTEGER PRIMARY KEY, package_event_id INTEGER NOT NULL, transaction_event_id_at_persist INTEGER NOT NULL, FOREIGN KEY(package_event_id) REFERENCES package_events(package_event_id)");
-	CREATE_TABLE("activities", "activity_id INTEGER PRIMARY KEY, endpoint_id BLOB NOT NULL, event_time INTEGER NOT NULL, event_type INTEGER NOT NULL, event_id INTEGER NOT NULL, event_summary_type TEXT NOT NULL, event_summary_flags INTEGER NOT NULL, event_summary_size_bytes INTEGER NOT NULL, event_summary_data BLOB, FOREIGN KEY(endpoint_id) REFERENCES endpoints(endpoint_id)");
+	CREATE_TABLE("activities", "activity_id INTEGER PRIMARY KEY, endpoint_id BLOB NOT NULL, event_time INTEGER NOT NULL, event_type INTEGER NOT NULL, event_id INTEGER NOT NULL, event_summary_type TEXT NOT NULL, event_summary_flags INTEGER NOT NULL, event_summary_size_bytes INTEGER NOT NULL, event_summary_data BLOB, activity_flags TINYINT UNSIGNED NOT NULL ,FOREIGN KEY(endpoint_id) REFERENCES endpoints(endpoint_id)");
 	CREATE_TABLE("ignored_activities", "activity_id INTEGER NOT NULL, FOREIGN KEY(activity_id) REFERENCES activities(activity_id)");
 	CREATE_TABLE("resource_locks", "object_name_id INTEGER NOT NULL, lock_event_id INTEGER NOT NULL, FOREIGN KEY(object_name_id) REFERENCES object_names(object_name_id), FOREIGN KEY(lock_event_id) REFERENCES lock_events(lock_event_id)");
 	CREATE_TABLE("package_transactions", "package_name_id INTEGER NOT NULL, transaction_event_id INTEGER NOT NULL, FOREIGN KEY(package_name_id) REFERENCES package_names(package_name_id), FOREIGN KEY(transaction_event_id) REFERENCES transaction_events(transaction_event_id)");
 	CREATE_TABLE("object_transactions", "object_name_id INTEGER NOT NULL, transaction_event_id INTEGER NOT NULL, FOREIGN KEY(object_name_id) REFERENCES object_names(object_name_id), FOREIGN KEY(transaction_event_id) REFERENCES transaction_events(transaction_event_id)");
 #undef CREATE_TABLE
 
+	if (LoadedDatabaseVersion < (int32)FConcertSyncSessionDatabaseVersion::AddActivityFlags
+		&& bActivitiesTableExisted)
+	{
+		if (!Database->Execute(TEXT("ALTER TABLE activities add activity_flags TINYINT UNSIGNED NOT NULL DEFAULT (0)")))
+		{
+			Close();
+			return false;
+		}
+	}
+
 	// Create our required indexes
 #define CREATE_INDEX(NAME, TABLE, COLS)																										\
 	if (!Database->Execute(TEXT("CREATE INDEX IF NOT EXISTS ") TEXT(NAME) TEXT(" ON ") TEXT(TABLE) TEXT("(") TEXT(COLS) TEXT(");")))		\
@@ -1699,7 +1736,7 @@ bool FConcertSyncSessionDatabase::AddConnectionActivity(const FConcertSyncConnec
 	FConcertSyncSessionDatabaseScopedTransaction ScopedTransaction(*Statements);
 	return ScopedTransaction.CommitOrRollback(
 		AddConnectionEvent(InConnectionActivity.EventData, OutConnectionEventId) &&
-		Statements->AddActivityData(InConnectionActivity.EndpointId, EConcertSyncActivityEventType::Connection, OutConnectionEventId, InConnectionActivity.EventSummary, OutActivityId) &&
+		Statements->AddActivityData(InConnectionActivity.Flags, InConnectionActivity.EndpointId, EConcertSyncActivityEventType::Connection, OutConnectionEventId, InConnectionActivity.EventSummary, OutActivityId) &&
 		SetActivityIgnoredState(OutActivityId, InConnectionActivity.bIgnored)
 		);
 }
@@ -1709,7 +1746,7 @@ bool FConcertSyncSessionDatabase::AddLockActivity(const FConcertSyncLockActivity
 	FConcertSyncSessionDatabaseScopedTransaction ScopedTransaction(*Statements);
 	return ScopedTransaction.CommitOrRollback(
 		AddLockEvent(InLockActivity.EventData, OutLockEventId) &&
-		Statements->AddActivityData(InLockActivity.EndpointId, EConcertSyncActivityEventType::Lock, OutLockEventId, InLockActivity.EventSummary, OutActivityId) &&
+		Statements->AddActivityData(InLockActivity.Flags, InLockActivity.EndpointId, EConcertSyncActivityEventType::Lock, OutLockEventId, InLockActivity.EventSummary, OutActivityId) &&
 		SetActivityIgnoredState(OutActivityId, InLockActivity.bIgnored)
 		);
 }
@@ -1719,7 +1756,7 @@ bool FConcertSyncSessionDatabase::AddTransactionActivity(const FConcertSyncTrans
 	FConcertSyncSessionDatabaseScopedTransaction ScopedTransaction(*Statements);
 	return ScopedTransaction.CommitOrRollback(
 		AddTransactionEvent(InTransactionActivity.EventData, OutTransactionEventId) &&
-		Statements->AddActivityData(InTransactionActivity.EndpointId, EConcertSyncActivityEventType::Transaction, OutTransactionEventId, InTransactionActivity.EventSummary, OutActivityId) &&
+		Statements->AddActivityData(InTransactionActivity.Flags, InTransactionActivity.EndpointId, EConcertSyncActivityEventType::Transaction, OutTransactionEventId, InTransactionActivity.EventSummary, OutActivityId) &&
 		SetActivityIgnoredState(OutActivityId, InTransactionActivity.bIgnored)
 		);
 }
@@ -1729,17 +1766,46 @@ bool FConcertSyncSessionDatabase::AddPackageActivity(const FConcertSyncActivity&
 	FConcertSyncSessionDatabaseScopedTransaction ScopedTransaction(*Statements);
 	return ScopedTransaction.CommitOrRollback(
 		AddPackageEvent(PackageInfo, InPackageDataStream, OutPackageEventId) &&
-		Statements->AddActivityData(InPackageActivityBasePart.EndpointId, EConcertSyncActivityEventType::Package, OutPackageEventId, InPackageActivityBasePart.EventSummary, OutActivityId) &&
+		Statements->AddActivityData(InPackageActivityBasePart.Flags, InPackageActivityBasePart.EndpointId, EConcertSyncActivityEventType::Package, OutPackageEventId, InPackageActivityBasePart.EventSummary, OutActivityId) &&
 		SetActivityIgnoredState(OutActivityId, InPackageActivityBasePart.bIgnored)
 		);
 }
 
+bool FConcertSyncSessionDatabase::SetActivities(const TSet<int64>& ActivityIds, TFunctionRef<void(FConcertSyncActivity&)> UpdateCallback)
+{
+	FConcertSyncSessionDatabaseScopedTransaction ScopedTransaction(*Statements);
+	bool bErrorFree = true;
+	ON_SCOPE_EXIT
+	{
+		if (bErrorFree)
+		{
+			ScopedTransaction.Commit();
+		}
+		else
+		{
+			ScopedTransaction.Rollback();
+		}
+	};
+	
+	for (auto ActivityIdIt = ActivityIds.CreateConstIterator(); ActivityIdIt && bErrorFree; ++ActivityIdIt)
+	{
+		FConcertSyncActivity Activity;
+		bErrorFree &= GetActivity(*ActivityIdIt, Activity);
+		if (bErrorFree)
+		{
+			UpdateCallback(Activity);
+			bErrorFree &= Statements->SetActivityData(Activity.ActivityId, Activity.Flags, Activity.EndpointId, Activity.EventTime, Activity.EventType, Activity.EventId, Activity.EventSummary);
+		}
+	}
+	return bErrorFree;
+}
+
 bool FConcertSyncSessionDatabase::SetConnectionActivity(const FConcertSyncConnectionActivity& InConnectionActivity)
 {
 	FConcertSyncSessionDatabaseScopedTransaction ScopedTransaction(*Statements);
 	return ScopedTransaction.CommitOrRollback(
 		SetConnectionEvent(InConnectionActivity.EventId, InConnectionActivity.EventData) &&
-		Statements->SetActivityData(InConnectionActivity.ActivityId, InConnectionActivity.EndpointId, InConnectionActivity.EventTime, InConnectionActivity.EventType, InConnectionActivity.EventId, InConnectionActivity.EventSummary) &&
+		Statements->SetActivityData(InConnectionActivity.ActivityId, InConnectionActivity.Flags, InConnectionActivity.EndpointId, InConnectionActivity.EventTime, InConnectionActivity.EventType, InConnectionActivity.EventId, InConnectionActivity.EventSummary) &&
 		SetActivityIgnoredState(InConnectionActivity.ActivityId, InConnectionActivity.bIgnored)
 		);
 }
@@ -1749,7 +1815,7 @@ bool FConcertSyncSessionDatabase::SetLockActivity(const FConcertSyncLockActivity
 	FConcertSyncSessionDatabaseScopedTransaction ScopedTransaction(*Statements);
 	return ScopedTransaction.CommitOrRollback(
 		SetLockEvent(InLockActivity.EventId, InLockActivity.EventData) &&
-		Statements->SetActivityData(InLockActivity.ActivityId, InLockActivity.EndpointId, InLockActivity.EventTime, InLockActivity.EventType, InLockActivity.EventId, InLockActivity.EventSummary) &&
+		Statements->SetActivityData(InLockActivity.ActivityId, InLockActivity.Flags, InLockActivity.EndpointId, InLockActivity.EventTime, InLockActivity.EventType, InLockActivity.EventId, InLockActivity.EventSummary) &&
 		SetActivityIgnoredState(InLockActivity.ActivityId, InLockActivity.bIgnored)
 		);
 }
@@ -1759,7 +1825,7 @@ bool FConcertSyncSessionDatabase::SetTransactionActivity(const FConcertSyncTrans
 	FConcertSyncSessionDatabaseScopedTransaction ScopedTransaction(*Statements);
 	return ScopedTransaction.CommitOrRollback(
 		SetTransactionEvent(InTransactionActivity.EventId, InTransactionActivity.EventData, bMetaDataOnly) &&
-		Statements->SetActivityData(InTransactionActivity.ActivityId, InTransactionActivity.EndpointId, InTransactionActivity.EventTime, InTransactionActivity.EventType, InTransactionActivity.EventId, InTransactionActivity.EventSummary) &&
+		Statements->SetActivityData(InTransactionActivity.ActivityId, InTransactionActivity.Flags, InTransactionActivity.EndpointId, InTransactionActivity.EventTime, InTransactionActivity.EventType, InTransactionActivity.EventId, InTransactionActivity.EventSummary) &&
 		SetActivityIgnoredState(InTransactionActivity.ActivityId, InTransactionActivity.bIgnored)
 		);
 }
@@ -1769,7 +1835,7 @@ bool FConcertSyncSessionDatabase::SetPackageActivity(const FConcertSyncActivity&
 	FConcertSyncSessionDatabaseScopedTransaction ScopedTransaction(*Statements);
 	return ScopedTransaction.CommitOrRollback(
 		SetPackageEvent(InPackageActivityBasePart.EventId, InPackageActivityEventPart.MetaData.PackageRevision, InPackageActivityEventPart.MetaData.PackageInfo, bMetaDataOnly ? static_cast<FConcertPackageDataStream*>(nullptr) : &InPackageActivityEventPart.PackageDataStream) &&
-		Statements->SetActivityData(InPackageActivityBasePart.ActivityId, InPackageActivityBasePart.EndpointId, InPackageActivityBasePart.EventTime, InPackageActivityBasePart.EventType, InPackageActivityBasePart.EventId, InPackageActivityBasePart.EventSummary) &&
+		Statements->SetActivityData(InPackageActivityBasePart.ActivityId, InPackageActivityBasePart.Flags, InPackageActivityBasePart.EndpointId, InPackageActivityBasePart.EventTime, InPackageActivityBasePart.EventType, InPackageActivityBasePart.EventId, InPackageActivityBasePart.EventSummary) &&
 		SetActivityIgnoredState(InPackageActivityBasePart.ActivityId, InPackageActivityBasePart.bIgnored)
 		);
 }
@@ -1778,7 +1844,7 @@ bool FConcertSyncSessionDatabase::GetActivity(const int64 InActivityId, FConcert
 {
 	OutActivity.ActivityId = InActivityId;
 	OutActivity.bIgnored = Statements->IsActivityIgnored(InActivityId);
-	return Statements->GetActivityDataForId(InActivityId, OutActivity.EndpointId, OutActivity.EventTime, OutActivity.EventType, OutActivity.EventId, OutActivity.EventSummary);
+	return Statements->GetActivityDataForId(InActivityId, OutActivity.Flags, OutActivity.EndpointId, OutActivity.EventTime, OutActivity.EventType, OutActivity.EventId, OutActivity.EventSummary);
 }
 
 bool FConcertSyncSessionDatabase::GetConnectionActivity(const int64 InActivityId, FConcertSyncConnectionActivity& OutConnectionActivity) const
@@ -1821,7 +1887,7 @@ bool FConcertSyncSessionDatabase::GetActivityForEvent(const int64 InEventId, con
 {
 	OutActivity.EventId = InEventId;
 	OutActivity.EventType = InEventType;
-	if (Statements->GetActivityDataForEvent(InEventId, InEventType, OutActivity.ActivityId, OutActivity.EndpointId, OutActivity.EventTime, OutActivity.EventSummary))
+	if (Statements->GetActivityDataForEvent(InEventId, InEventType, OutActivity.ActivityId, OutActivity.Flags, OutActivity.EndpointId, OutActivity.EventTime, OutActivity.EventSummary))
 	{
 		OutActivity.bIgnored = Statements->IsActivityIgnored(OutActivity.ActivityId);
 		return true;
@@ -1862,11 +1928,12 @@ bool FConcertSyncSessionDatabase::GetPackageActivityForEvent(const int64 InPacka
 
 bool FConcertSyncSessionDatabase::EnumerateActivities(FIterateActivityFunc InCallback) const
 {
-	return Statements->GetAllActivityData([this, &InCallback](const int64 InActivityId, const FGuid& InEndpointId, const FDateTime InEventTime, const EConcertSyncActivityEventType InEventType, const int64 InEventId, FConcertSessionSerializedPayload&& InEventSummary)
+	return Statements->GetAllActivityData([this, &InCallback](const int64 InActivityId, const EConcertSyncActivityFlags InFlags, const FGuid& InEndpointId, const FDateTime InEventTime, const EConcertSyncActivityEventType InEventType, const int64 InEventId, FConcertSessionSerializedPayload&& InEventSummary)
 	{
 		FConcertSyncActivity Activity;
 		Activity.ActivityId = InActivityId;
 		Activity.bIgnored = Statements->IsActivityIgnored(InActivityId);
+		Activity.Flags = InFlags;
 		Activity.EndpointId = InEndpointId;
 		Activity.EventTime = InEventTime;
 		Activity.EventType = InEventType;
@@ -1880,11 +1947,12 @@ bool FConcertSyncSessionDatabase::EnumerateActivities(FIterateActivityFunc InCal
 
 bool FConcertSyncSessionDatabase::EnumerateConnectionActivities(TFunctionRef<bool(FConcertSyncConnectionActivity&&)> InCallback) const
 {
-	return Statements->GetAllActivityDataForEventType(EConcertSyncActivityEventType::Connection, [this, &InCallback](const int64 InActivityId, const FGuid& InEndpointId, const FDateTime InEventTime, const int64 InEventId, FConcertSessionSerializedPayload&& InEventSummary)
+	return Statements->GetAllActivityDataForEventType(EConcertSyncActivityEventType::Connection, [this, &InCallback](const int64 InActivityId, const EConcertSyncActivityFlags InFlags, const FGuid& InEndpointId, const FDateTime InEventTime, const int64 InEventId, FConcertSessionSerializedPayload&& InEventSummary)
 	{
 		FConcertSyncConnectionActivity ConnectionActivity;
 		ConnectionActivity.ActivityId = InActivityId;
 		ConnectionActivity.bIgnored = Statements->IsActivityIgnored(InActivityId);
+		ConnectionActivity.Flags = InFlags;
 		ConnectionActivity.EndpointId = InEndpointId;
 		ConnectionActivity.EventTime = InEventTime;
 		ConnectionActivity.EventType = EConcertSyncActivityEventType::Connection;
@@ -1902,11 +1970,12 @@ bool FConcertSyncSessionDatabase::EnumerateConnectionActivities(TFunctionRef<boo
 
 bool FConcertSyncSessionDatabase::EnumerateLockActivities(TFunctionRef<bool(FConcertSyncLockActivity&&)> InCallback) const
 {
-	return Statements->GetAllActivityDataForEventType(EConcertSyncActivityEventType::Lock, [this, &InCallback](const int64 InActivityId, const FGuid& InEndpointId, const FDateTime InEventTime, const int64 InEventId, FConcertSessionSerializedPayload&& InEventSummary)
+	return Statements->GetAllActivityDataForEventType(EConcertSyncActivityEventType::Lock, [this, &InCallback](const int64 InActivityId, const EConcertSyncActivityFlags InFlags, const FGuid& InEndpointId, const FDateTime InEventTime, const int64 InEventId, FConcertSessionSerializedPayload&& InEventSummary)
 	{
 		FConcertSyncLockActivity LockActivity;
 		LockActivity.ActivityId = InActivityId;
 		LockActivity.bIgnored = Statements->IsActivityIgnored(InActivityId);
+		LockActivity.Flags = InFlags;
 		LockActivity.EndpointId = InEndpointId;
 		LockActivity.EventTime = InEventTime;
 		LockActivity.EventType = EConcertSyncActivityEventType::Lock;
@@ -1924,11 +1993,12 @@ bool FConcertSyncSessionDatabase::EnumerateLockActivities(TFunctionRef<bool(FCon
 
 bool FConcertSyncSessionDatabase::EnumerateTransactionActivities(TFunctionRef<bool(FConcertSyncTransactionActivity&&)> InCallback) const
 {
-	return Statements->GetAllActivityDataForEventType(EConcertSyncActivityEventType::Transaction, [this, &InCallback](const int64 InActivityId, const FGuid& InEndpointId, const FDateTime InEventTime, const int64 InEventId, FConcertSessionSerializedPayload&& InEventSummary)
+	return Statements->GetAllActivityDataForEventType(EConcertSyncActivityEventType::Transaction, [this, &InCallback](const int64 InActivityId, const EConcertSyncActivityFlags InFlags, const FGuid& InEndpointId, const FDateTime InEventTime, const int64 InEventId, FConcertSessionSerializedPayload&& InEventSummary)
 	{
 		FConcertSyncTransactionActivity TransactionActivity;
 		TransactionActivity.ActivityId = InActivityId;
 		TransactionActivity.bIgnored = Statements->IsActivityIgnored(InActivityId);
+		TransactionActivity.Flags = InFlags;
 		TransactionActivity.EndpointId = InEndpointId;
 		TransactionActivity.EventTime = InEventTime;
 		TransactionActivity.EventType = EConcertSyncActivityEventType::Transaction;
@@ -1946,11 +2016,12 @@ bool FConcertSyncSessionDatabase::EnumerateTransactionActivities(TFunctionRef<bo
 
 bool FConcertSyncSessionDatabase::EnumeratePackageActivities(FIteratePackageActivityFunc InCallback) const
 {
-	return Statements->GetAllActivityDataForEventType(EConcertSyncActivityEventType::Package, [this, &InCallback](const int64 InActivityId, const FGuid& InEndpointId, const FDateTime InEventTime, const int64 InEventId, FConcertSessionSerializedPayload&& InEventSummary)
+	return Statements->GetAllActivityDataForEventType(EConcertSyncActivityEventType::Package, [this, &InCallback](const int64 InActivityId, const EConcertSyncActivityFlags InFlags, const FGuid& InEndpointId, const FDateTime InEventTime, const int64 InEventId, FConcertSessionSerializedPayload&& InEventSummary)
 	{
 		FConcertSyncActivity PackageActivity;
 		PackageActivity.ActivityId = InActivityId;
 		PackageActivity.bIgnored = Statements->IsActivityIgnored(InActivityId);
+		PackageActivity.Flags = InFlags;
 		PackageActivity.EndpointId = InEndpointId;
 		PackageActivity.EventTime = InEventTime;
 		PackageActivity.EventType = EConcertSyncActivityEventType::Package;
@@ -1970,11 +2041,12 @@ bool FConcertSyncSessionDatabase::EnumeratePackageActivities(FIteratePackageActi
 
 bool FConcertSyncSessionDatabase::EnumerateActivitiesForEventType(const EConcertSyncActivityEventType InEventType, FIterateActivityFunc InCallback) const
 {
-	return Statements->GetAllActivityDataForEventType(InEventType, [this, InEventType, &InCallback](const int64 InActivityId, const FGuid& InEndpointId, const FDateTime InEventTime, const int64 InEventId, FConcertSessionSerializedPayload&& InEventSummary)
+	return Statements->GetAllActivityDataForEventType(InEventType, [this, InEventType, &InCallback](const int64 InActivityId, const EConcertSyncActivityFlags InFlags, const FGuid& InEndpointId, const FDateTime InEventTime, const int64 InEventId, FConcertSessionSerializedPayload&& InEventSummary)
 	{
 		FConcertSyncActivity Activity;
 		Activity.ActivityId = InActivityId;
 		Activity.bIgnored = Statements->IsActivityIgnored(InActivityId);
+		Activity.Flags = InFlags;
 		Activity.EndpointId = InEndpointId;
 		Activity.EventTime = InEventTime;
 		Activity.EventType = InEventType;
@@ -1988,11 +2060,12 @@ bool FConcertSyncSessionDatabase::EnumerateActivitiesForEventType(const EConcert
 
 bool FConcertSyncSessionDatabase::EnumerateActivitiesInRange(const int64 InFirstActivityId, const int64 InMaxNumActivities, FIterateActivityFunc InCallback) const
 {
-	return Statements->GetActivityDataInRange(InFirstActivityId, InMaxNumActivities, [this, &InCallback](const int64 InActivityId, const FGuid& InEndpointId, const FDateTime InEventTime, const EConcertSyncActivityEventType InEventType, const int64 InEventId, FConcertSessionSerializedPayload&& InEventSummary)
+	return Statements->GetActivityDataInRange(InFirstActivityId, InMaxNumActivities, [this, &InCallback](const int64 InActivityId, const EConcertSyncActivityFlags InFlags, const FGuid& InEndpointId, const FDateTime InEventTime, const EConcertSyncActivityEventType InEventType, const int64 InEventId, FConcertSessionSerializedPayload&& InEventSummary)
 	{
 		FConcertSyncActivity Activity;
 		Activity.ActivityId = InActivityId;
 		Activity.bIgnored = Statements->IsActivityIgnored(InActivityId);
+		Activity.Flags = InFlags;
 		Activity.EndpointId = InEndpointId;
 		Activity.EventTime = InEventTime;
 		Activity.EventType = InEventType;
@@ -2024,6 +2097,16 @@ bool FConcertSyncSessionDatabase::EnumerateActivityIdsAndEventTypesInRange(const
 	});
 }
 
+bool FConcertSyncSessionDatabase::EnumerateActivityIdsWithEventTypesAndFlagsInRange(const int64 InFirstActivityId, const int64 InMaxNumActivities, TFunctionRef<bool(int64, EConcertSyncActivityEventType, EConcertSyncActivityFlags)> InCallback) const
+{
+	return Statements->GetActivityIdWithEventTypesAndFlagsInRange(InFirstActivityId, InMaxNumActivities, [&InCallback](const int64 InActivityId, const EConcertSyncActivityEventType InEventType, EConcertSyncActivityFlags InFlags)
+	{
+		return InCallback(InActivityId, InEventType, InFlags)
+			? ESQLitePreparedStatementExecuteRowResult::Continue
+			: ESQLitePreparedStatementExecuteRowResult::Stop;
+	});
+}
+
 bool FConcertSyncSessionDatabase::GetActivityMaxId(int64& OutActivityId) const
 {
 	return Statements->GetActivityMaxId(OutActivityId);
diff --git a/Engine/Plugins/Developer/Concert/ConcertSync/ConcertSyncCore/Source/ConcertSyncCore/Public/ConcertSyncSessionDatabase.h b/Engine/Plugins/Developer/Concert/ConcertSync/ConcertSyncCore/Source/ConcertSyncCore/Public/ConcertSyncSessionDatabase.h
--- a/Engine/Plugins/Developer/Concert/ConcertSync/ConcertSyncCore/Source/ConcertSyncCore/Public/ConcertSyncSessionDatabase.h
+++ b/Engine/Plugins/Developer/Concert/ConcertSync/ConcertSyncCore/Source/ConcertSyncCore/Public/ConcertSyncSessionDatabase.h
@@ -137,6 +137,11 @@ public:
 	 */
 	bool AddPackageActivity(const FConcertSyncActivity& InPackageActivity, const FConcertPackageInfo& InPackageInfo, FConcertPackageDataStream& InPackageDataStream, int64& OutActivityId, int64& OutPackageEventId);
 
+	/**
+	 * Iterates the given range, calls UpdateCallback on each element, and commits the update.
+	 */
+	bool SetActivities(const TSet<int64>& ActivityIds, TFunctionRef<void(FConcertSyncActivity&)> UpdateCallback);
+	
 	/**
 	 * Set a connection activity in this database, creating or replacing it.
 	 * @note The endpoint ID referenced by the activity must exist in the database (@see SetEndpoint).
@@ -383,6 +388,17 @@ public:
 	 */
 	bool EnumerateActivityIdsAndEventTypesInRange(const int64 InFirstActivityId, const int64 InMaxNumActivities, TFunctionRef<bool(int64, EConcertSyncActivityEventType)> InCallback) const;
 
+	/**
+	 * Enumerate the IDs, event types and flags of the activities in this database in the given range.
+	 *
+	 * @param InFirstActivityId			The first activity ID to include in the results.
+	 * @param InMaxNumActivities		The maximum number of activities to include in the results.
+	 * @param InCallback				Callback invoked for each activity; return true to continue enumeration, or false to stop.
+	 *
+	 * @return True if the activities were enumerated without error, false otherwise.
+	 */
+	bool EnumerateActivityIdsWithEventTypesAndFlagsInRange(const int64 InFirstActivityId, const int64 InMaxNumActivities, TFunctionRef<bool(int64, EConcertSyncActivityEventType, EConcertSyncActivityFlags)> InCallback) const;
+	
 	/**
 	 * Get the maximum ID of the activities in this database.
 	 *
diff --git a/Engine/Plugins/Developer/Concert/ConcertSync/ConcertSyncCore/Source/ConcertSyncCore/Public/ConcertSyncSessionTypes.h b/Engine/Plugins/Developer/Concert/ConcertSync/ConcertSyncCore/Source/ConcertSyncCore/Public/ConcertSyncSessionTypes.h
--- a/Engine/Plugins/Developer/Concert/ConcertSync/ConcertSyncCore/Source/ConcertSyncCore/Public/ConcertSyncSessionTypes.h
+++ b/Engine/Plugins/Developer/Concert/ConcertSync/ConcertSyncCore/Source/ConcertSyncCore/Public/ConcertSyncSessionTypes.h
@@ -162,6 +162,18 @@ struct FConcertSyncPackageEventData
 	FConcertPackageDataStream PackageDataStream;
 };
 
+UENUM()
+enum class EConcertSyncActivityFlags : uint8
+{
+	None = 0,
+	/**
+	 * This activity will never be sent to clients by the server.
+	 * For all activities a client receives (Flags & EConcertSyncActivityFlags::Muted) == EConcertSyncActivityFlags::None holds. 
+	 */
+	Muted = 1 << 0,
+};
+ENUM_CLASS_FLAGS(EConcertSyncActivityFlags)
+
 /** Data for an activity entry in a Concert Sync Session */
 USTRUCT()
 struct FConcertSyncActivity
@@ -176,6 +188,10 @@ struct FConcertSyncActivity
 	UPROPERTY()
 	bool bIgnored = false;
 
+	/** Additional information about this activity */
+	UPROPERTY()
+	EConcertSyncActivityFlags Flags = EConcertSyncActivityFlags::None;
+
 	/** The ID of the endpoint that produced the activity */
 	UPROPERTY()
 	FGuid EndpointId;
diff --git a/Engine/Plugins/Developer/Concert/ConcertSync/ConcertSyncServer/Source/ConcertSyncServer/Private/ConcertServerEventForwardingSink.h b/Engine/Plugins/Developer/Concert/ConcertSync/ConcertSyncServer/Source/ConcertSyncServer/Private/ConcertServerEventForwardingSink.h
--- a/Engine/Plugins/Developer/Concert/ConcertSync/ConcertSyncServer/Source/ConcertSyncServer/Private/ConcertServerEventForwardingSink.h
+++ b/Engine/Plugins/Developer/Concert/ConcertSync/ConcertSyncServer/Source/ConcertSyncServer/Private/ConcertServerEventForwardingSink.h
@@ -86,7 +86,7 @@ public:
 		return bSuccess;
 	}
 
-	virtual bool GetSessionActivities(const IConcertServer& InServer, const FGuid& SessionId, int64 FromActivityId, int64 ActivityCount, TArray<FConcertSessionSerializedPayload>& OutActivities, TMap<FGuid, FConcertClientInfo>& OutEndpointClientInfoMap, bool bIncludeDetails) final
+	virtual bool GetUnmutedSessionActivities(const IConcertServer& InServer, const FGuid& SessionId, int64 FromActivityId, int64 ActivityCount, TArray<FConcertSessionSerializedPayload>& OutActivities, TMap<FGuid, FConcertClientInfo>& OutEndpointClientInfoMap, bool bIncludeDetails) final
 	{
 		return This()->GetSessionActivitiesImpl(InServer, SessionId, FromActivityId, ActivityCount, OutActivities, OutEndpointClientInfoMap, bIncludeDetails);
 	}
diff --git a/Engine/Plugins/Developer/Concert/ConcertSync/ConcertSyncServer/Source/ConcertSyncServer/Private/ConcertServerWorkspace.cpp b/Engine/Plugins/Developer/Concert/ConcertSync/ConcertSyncServer/Source/ConcertSyncServer/Private/ConcertServerWorkspace.cpp
--- a/Engine/Plugins/Developer/Concert/ConcertSync/ConcertSyncServer/Source/ConcertSyncServer/Private/ConcertServerWorkspace.cpp
+++ b/Engine/Plugins/Developer/Concert/ConcertSync/ConcertSyncServer/Source/ConcertSyncServer/Private/ConcertServerWorkspace.cpp
@@ -227,8 +227,13 @@ void FConcertServerWorkspace::HandleSyncRequestedEvent(const FConcertSessionCont
 	});
 
 	// Sync all activity
-	LiveSession->GetSessionDatabase().EnumerateActivityIdsAndEventTypesInRange(FirstActivityIdToSync, NumActivitiesToSync, [this, &Context](const int64 InActivityId, const EConcertSyncActivityEventType InEventType)
+	LiveSession->GetSessionDatabase().EnumerateActivityIdsWithEventTypesAndFlagsInRange(FirstActivityIdToSync, NumActivitiesToSync, [this, &Context](const int64 InActivityId, const EConcertSyncActivityEventType InEventType, const EConcertSyncActivityFlags InFlags)
 	{
+		if ((InFlags & EConcertSyncActivityFlags::Muted) != EConcertSyncActivityFlags::None)
+		{
+			return true;
+		}
+		
 		SyncCommandQueue->QueueCommand(Context.SourceEndpointId, [this, SyncActivityId = InActivityId, SyncEventType = InEventType](const FConcertServerSyncCommandQueue::FSyncCommandContext& InSyncCommandContext, const FGuid& InEndpointId)
 		{
 			switch (SyncEventType)
diff --git a/Engine/Plugins/Developer/Concert/ConcertSync/ConcertSyncServer/Source/ConcertSyncServer/Private/ConcertSyncServer.cpp b/Engine/Plugins/Developer/Concert/ConcertSync/ConcertSyncServer/Source/ConcertSyncServer/Private/ConcertSyncServer.cpp
--- a/Engine/Plugins/Developer/Concert/ConcertSync/ConcertSyncServer/Source/ConcertSyncServer/Private/ConcertSyncServer.cpp
+++ b/Engine/Plugins/Developer/Concert/ConcertSync/ConcertSyncServer/Source/ConcertSyncServer/Private/ConcertSyncServer.cpp
@@ -416,7 +416,7 @@ bool FConcertSyncServer::GetSessionActivitiesInternal(const FConcertSyncSessionD
 	OutActivities.Reset(FMath::Min(ActivityCount, MaxActivityId));
 
 	// Retrieve the generic part of activities.
-	Database.EnumerateActivityIdsAndEventTypesInRange(FromActivityId, ActivityCount, [&Database, &OutActivities, &OutEndpointClientInfoMap, bIncludeDetails](const int64 InActivityId, const EConcertSyncActivityEventType InEventType)
+	Database.EnumerateActivityIdsWithEventTypesAndFlagsInRange(FromActivityId, ActivityCount, [&Database, &OutActivities, &OutEndpointClientInfoMap, bIncludeDetails](const int64 InActivityId, const EConcertSyncActivityEventType InEventType, const EConcertSyncActivityFlags InFlags)
 	{
 		// Maps endpoint client id to the client info.
 		auto UpdateEndpointMap = [](const FConcertSyncSessionDatabase& Database, FGuid EndpointId, TMap<FGuid, FConcertClientInfo>& OutEndpointClientInfoMap)
@@ -473,7 +473,11 @@ bool FConcertSyncServer::GetSessionActivitiesInternal(const FConcertSyncSessionD
 			}
 		}
 
-		OutActivities.Add(MoveTemp(SerializedSyncActivityPayload));
+		const bool bIsAllowed = (InFlags & EConcertSyncActivityFlags::Muted) == EConcertSyncActivityFlags::None; 
+		if (bIsAllowed)
+		{
+			OutActivities.Add(MoveTemp(SerializedSyncActivityPayload));
+		}
 
 		return true; // Continue until 'ActivityCount' is fetched or the last activity is reached.
 	});
