Fix SQL statements to respect muted activities

diff --git a/Engine/Plugins/Developer/Concert/ConcertSync/ConcertSyncCore/Source/ConcertSyncCore/Private/ConcertSyncSessionDatabase.cpp b/Engine/Plugins/Developer/Concert/ConcertSync/ConcertSyncCore/Source/ConcertSyncCore/Private/ConcertSyncSessionDatabase.cpp
--- a/Engine/Plugins/Developer/Concert/ConcertSync/ConcertSyncCore/Source/ConcertSyncCore/Private/ConcertSyncSessionDatabase.cpp
+++ b/Engine/Plugins/Developer/Concert/ConcertSync/ConcertSyncCore/Source/ConcertSyncCore/Private/ConcertSyncSessionDatabase.cpp
@@ -634,7 +634,6 @@ public:
 
 		PREPARE_STATEMENT(Statement_MapPackageNameIdToTransactionEventId);
 		PREPARE_STATEMENT(Statement_UnmapPackageNameIdsForTransactionEventId);
-		PREPARE_STATEMENT(Statement_GetTransactionEventIdsForPackageNameId);
 		PREPARE_STATEMENT(Statement_GetTransactionEventIdsInRangeForPackageNameId);
 		PREPARE_STATEMENT(Statement_GetPackageNameIdsMaxTransactionId);
 		PREPARE_STATEMENT(Statement_GetPackageNameIdsWithTransactions);
@@ -642,7 +641,6 @@ public:
 
 		PREPARE_STATEMENT(Statement_MapObjectNameIdToTransactionEventId);
 		PREPARE_STATEMENT(Statement_UnmapObjectNameIdsForTransactionEventId);
-		PREPARE_STATEMENT(Statement_GetTransactionEventIdsForObjectNameId);
 
 		#undef PREPARE_STATEMENT
 
@@ -965,7 +963,7 @@ public:
 		}
 		return false;
 	}
-
+	
 	/** Get the package_name_id and package_revision from package_events for the given package_event_id */
 	SQLITE_PREPARED_STATEMENT(FGetPackageNameIdAndRevisonForId, "SELECT package_name_id, package_revision FROM package_events WHERE package_event_id = ?1;", SQLITE_PREPARED_STATEMENT_COLUMNS(int64, int64), SQLITE_PREPARED_STATEMENT_BINDINGS(int64));
 	FGetPackageNameIdAndRevisonForId Statement_GetPackageNameIdAndRevisonForId;
@@ -1030,7 +1028,7 @@ public:
 	{
 		return Statement_GetPackageHeadEventIdAndTransactionIdAtSave.BindAndExecuteSingle(InPackageNameId, OutPackageEventId, OutTransactionEventIdAtSave);
 	}
-
+	
 	/** Get the largest package_event_id along its transaction_event_id_at_save currently in package_events for each distinct package_name_id */
 	SQLITE_PREPARED_STATEMENT_COLUMNS_ONLY(FGetMaxPackageEventIdAndTransactionEventIdAtSavePerPackageNameId, "SELECT package_name_id, MAX(package_event_id), transaction_event_id_at_save FROM package_events GROUP BY package_name_id;", SQLITE_PREPARED_STATEMENT_COLUMNS(int64, int64, int64));
 	FGetMaxPackageEventIdAndTransactionEventIdAtSavePerPackageNameId Statement_GetMaxPackageEventIdAndTransactionEventIdAtSavePerPackageNameId;
@@ -1047,18 +1045,21 @@ public:
 		}) != INDEX_NONE;
 	}
 
-	/** Get the largest package_revision currently in package_events for the given package_name_id */
-	SQLITE_PREPARED_STATEMENT(FGetPackageHeadRevison, "SELECT MAX(package_revision) FROM package_events WHERE package_name_id = ?1;", SQLITE_PREPARED_STATEMENT_COLUMNS(int64), SQLITE_PREPARED_STATEMENT_BINDINGS(int64));
+	/** Get the largest package_revision (for which the activity is not muted) currently in package_events for the given package_name_id */
+	SQLITE_PREPARED_STATEMENT(FGetPackageHeadRevison, "SELECT MAX(package_revision) FROM package_events INNER JOIN activities ON package_events.package_event_id = activities.event_id WHERE activities.activity_flags & 1 = 0 AND activities.event_type = 4 AND package_name_id = ?1", SQLITE_PREPARED_STATEMENT_COLUMNS(int64), SQLITE_PREPARED_STATEMENT_BINDINGS(int64));
 	FGetPackageHeadRevison Statement_GetPackageHeadRevison;
-	bool GetPackageHeadRevision(const int64 InPackageNameId, int64& OutRevision)
+	static_assert(static_cast<int32>(EConcertSyncActivityFlags::Muted) == 1, "Please update the SQL statement");
+	bool GetUnmutedPackageHeadRevision(const int64 InPackageNameId, int64& OutRevision)
 	{
 		return Statement_GetPackageHeadRevison.BindAndExecuteSingle(InPackageNameId, OutRevision);
 	}
 
-	/** Get the largest transaction_event_id_at_save currently in package_events for the given package_name_id */
-	SQLITE_PREPARED_STATEMENT(FGetPackageTransactionEventIdAtLastSave, "SELECT MAX(transaction_event_id_at_save) FROM package_events WHERE package_name_id = ?1;", SQLITE_PREPARED_STATEMENT_COLUMNS(int64), SQLITE_PREPARED_STATEMENT_BINDINGS(int64));
+	/** Get the largest transaction_event_id_at_save (for which the transaction is not muted) currently in package_events for the given package_name_id */
+	SQLITE_PREPARED_STATEMENT(FGetPackageTransactionEventIdAtLastSave, "SELECT MAX(package_events.transaction_event_id_at_save) FROM package_events INNER JOIN activities ON package_events.package_event_id = activities.event_id WHERE activities.activity_flags & 1 = 0 AND activities.event_type = 4 AND package_name_id = ?1;", SQLITE_PREPARED_STATEMENT_COLUMNS(int64), SQLITE_PREPARED_STATEMENT_BINDINGS(int64));
 	FGetPackageTransactionEventIdAtLastSave Statement_GetPackageTransactionEventIdAtLastSave;
-	bool GetPackageTransactionEventIdAtLastSave(const int64 InPackageNameId, int64& OutTransactionEventId)
+	static_assert(static_cast<int32>(EConcertSyncActivityFlags::Muted) == 1, "Update activities.activity_flags & 1 = 0");
+	static_assert(static_cast<int32>(EConcertSyncActivityEventType::Package) == 4, "Update activities.event_type = 4");
+	bool GetUnmutedPackageTransactionEventIdAtLastSave(const int64 InPackageNameId, int64& OutTransactionEventId)
 	{
 		return Statement_GetPackageTransactionEventIdAtLastSave.BindAndExecuteSingle(InPackageNameId, OutTransactionEventId);
 	}
@@ -1385,25 +1386,10 @@ public:
 		return Statement_UnmapPackageNameIdsForTransactionEventId.BindAndExecute(InTransactionEventId);
 	}
 
-	/** Get the transaction_event_id entries from package_transactions for the given package_name_id */
-	SQLITE_PREPARED_STATEMENT(FGetTransactionEventIdsForPackageNameId, "SELECT transaction_event_id FROM package_transactions WHERE package_name_id = ?1 ORDER BY transaction_event_id;", SQLITE_PREPARED_STATEMENT_COLUMNS(int64), SQLITE_PREPARED_STATEMENT_BINDINGS(int64));
-	FGetTransactionEventIdsForPackageNameId Statement_GetTransactionEventIdsForPackageNameId;
-	bool GetTransactionIdsForPackageNameId(const int64 InPackageNameId, TFunctionRef<ESQLitePreparedStatementExecuteRowResult(int64)> InCallback)
-	{
-		return Statement_GetTransactionEventIdsForPackageNameId.BindAndExecute(InPackageNameId, [&InCallback](const FGetTransactionEventIdsForPackageNameId& InStatement)
-		{
-			int64 TransactionEventId = 0;
-			if (InStatement.GetColumnValues(TransactionEventId))
-			{
-				return InCallback(TransactionEventId);
-			}
-			return ESQLitePreparedStatementExecuteRowResult::Error;
-		}) != INDEX_NONE;
-	}
-
 	/** Get the transaction_event_id entries from package_transactions for the given package_name_id and a transaction_event_id >= the given mininum transaction_event_id */
 	SQLITE_PREPARED_STATEMENT(FGetTransactionEventIdsInRangeForPackageNameId, "SELECT transaction_event_id FROM package_transactions WHERE package_name_id = ?1 AND transaction_event_id >= ?2 ORDER BY transaction_event_id;", SQLITE_PREPARED_STATEMENT_COLUMNS(int64), SQLITE_PREPARED_STATEMENT_BINDINGS(int64, int64));
 	FGetTransactionEventIdsInRangeForPackageNameId Statement_GetTransactionEventIdsInRangeForPackageNameId;
+	static_assert(static_cast<int32>(EConcertSyncActivityFlags::Muted) == 1, "Please update the SQL statement");
 	bool GetTransactionEventIdsInRangeForPackageNameId(const int64 InPackageNameId, const int64 InMinTransactionEventId, TFunctionRef<ESQLitePreparedStatementExecuteRowResult(int64)> InCallback)
 	{
 		return Statement_GetTransactionEventIdsInRangeForPackageNameId.BindAndExecute(InPackageNameId, InMinTransactionEventId, [&InCallback](const FGetTransactionEventIdsInRangeForPackageNameId& InStatement)
@@ -1420,6 +1406,7 @@ public:
 	/** Get the max transaction_event_id for each package_name_ids from package_transactions */
 	SQLITE_PREPARED_STATEMENT_COLUMNS_ONLY(FGetPackageNameIdsMaxTransactionId, "SELECT package_name_id, MAX(transaction_event_id) FROM package_transactions GROUP BY package_name_id;", SQLITE_PREPARED_STATEMENT_COLUMNS(int64, int64));
 	FGetPackageNameIdsMaxTransactionId Statement_GetPackageNameIdsMaxTransactionId;
+	static_assert(static_cast<int32>(EConcertSyncActivityFlags::Muted) == 1, "Please update the SQL statement");
 	bool GetPackageNameIdsMaxTransactionId(TFunctionRef<ESQLitePreparedStatementExecuteRowResult(int64, int64)> InCallback)
 	{
 		return Statement_GetPackageNameIdsMaxTransactionId.Execute([&InCallback](const FGetPackageNameIdsMaxTransactionId& InStatement)
@@ -1485,22 +1472,6 @@ public:
 		return Statement_UnmapObjectNameIdsForTransactionEventId.BindAndExecute(InTransactionEventId);
 	}
 
-	/** Get the transaction_event_id entries from object_transactions for the given object_name_id */
-	SQLITE_PREPARED_STATEMENT(FGetTransactionEventIdsForObjectNameId, "SELECT transaction_event_id FROM object_transactions WHERE object_name_id = ?1 ORDER BY transaction_event_id;", SQLITE_PREPARED_STATEMENT_COLUMNS(int64), SQLITE_PREPARED_STATEMENT_BINDINGS(int64));
-	FGetTransactionEventIdsForObjectNameId Statement_GetTransactionEventIdsForObjectNameId;
-	bool GetTransactionEventIdsForObjectNameId(const int64 InPackageNameId, TFunctionRef<ESQLitePreparedStatementExecuteRowResult(int64)> InCallback)
-	{
-		return Statement_GetTransactionEventIdsForObjectNameId.BindAndExecute(InPackageNameId, [&InCallback](const FGetTransactionEventIdsForObjectNameId& InStatement)
-		{
-			int64 TransactionEventId = 0;
-			if (InStatement.GetColumnValues(TransactionEventId))
-			{
-				return InCallback(TransactionEventId);
-			}
-			return ESQLitePreparedStatementExecuteRowResult::Error;
-		}) != INDEX_NONE;
-	}
-
 private:
 	FSQLiteDatabase& Database;
 };
@@ -2285,7 +2256,7 @@ bool FConcertSyncSessionDatabase::IsLiveTransactionEvent(const int64 InTransacti
 	return Statements->GetPackageNameIdsForTransactionEventId(InTransactionEventId, [this, InTransactionEventId, &OutIsLive](const int64 InPackageNameId)
 	{
 		int64 HeadTransactionEventIdAtLastSave = 0;
-		if (!Statements->GetPackageTransactionEventIdAtLastSave(InPackageNameId, HeadTransactionEventIdAtLastSave))
+		if (!Statements->GetUnmutedPackageTransactionEventIdAtLastSave(InPackageNameId, HeadTransactionEventIdAtLastSave))
 		{
 			return ESQLitePreparedStatementExecuteRowResult::Error;
 		}
@@ -2357,23 +2328,13 @@ bool FConcertSyncSessionDatabase::EnumerateLiveTransactionEventIdsForPackage(con
 	return EnumerateLiveTransactionEventIdsForPackage(PackageNameId, InCallback);
 }
 
-bool FConcertSyncSessionDatabase::GetPackageNamesWithLiveTransactions(TArray<FName>& OutPackageNames) const
-{
-	OutPackageNames.Reset();
-	return EnumeratePackageNamesWithLiveTransactions([&OutPackageNames](const FName InPackageName)
-	{
-		OutPackageNames.Add(InPackageName);
-		return true;
-	});
-}
-
 bool FConcertSyncSessionDatabase::EnumeratePackageNamesWithLiveTransactions(TFunctionRef<bool(FName)> InCallback) const
 {
 	return Statements->GetPackageNameIdsMaxTransactionId([this, &InCallback](const int64 InPackageNameId, const int64 InMaxTransactionEventId)
 	{
 		// Get the transaction id at last save, if the max transaction id for a package name id is greater than its transaction id at last save, it has live transactions
 		int64 HeadTransactionEventIdAtLastSave = 0;
-		Statements->GetPackageTransactionEventIdAtLastSave(InPackageNameId, HeadTransactionEventIdAtLastSave);
+		Statements->GetUnmutedPackageTransactionEventIdAtLastSave(InPackageNameId, HeadTransactionEventIdAtLastSave);
 		
 		if (InMaxTransactionEventId > HeadTransactionEventIdAtLastSave)
 		{
@@ -2393,7 +2354,7 @@ bool FConcertSyncSessionDatabase::EnumeratePackageNamesWithLiveTransactions(TFun
 bool FConcertSyncSessionDatabase::EnumerateLiveTransactionEventIdsForPackage(const int64 InPackageNameId, TFunctionRef<bool(int64)> InCallback) const
 {
 	int64 HeadTransactionEventIdAtLastSave = 0;
-	if (!Statements->GetPackageTransactionEventIdAtLastSave(InPackageNameId, HeadTransactionEventIdAtLastSave) || HeadTransactionEventIdAtLastSave == MAX_int64)
+	if (!Statements->GetUnmutedPackageTransactionEventIdAtLastSave(InPackageNameId, HeadTransactionEventIdAtLastSave) || HeadTransactionEventIdAtLastSave == MAX_int64)
 	{
 		return false;
 	}
@@ -2540,16 +2501,6 @@ bool FConcertSyncSessionDatabase::GetPackageEvent(const int64 InPackageEventId,
 	return false;
 }
 
-bool FConcertSyncSessionDatabase::GetPackageNamesWithHeadRevision(TArray<FName>& OutPackageNames, bool IgnorePersisted) const
-{
-	OutPackageNames.Reset();
-	return EnumeratePackageNamesWithHeadRevision([&OutPackageNames](const FName InPackageName)
-	{
-		OutPackageNames.Add(InPackageName);
-		return true;
-	}, IgnorePersisted);
-}
-
 bool FConcertSyncSessionDatabase::EnumeratePackageNamesWithHeadRevision(TFunctionRef<bool(FName)> InCallback, bool IgnorePersisted) const
 {
 	// if we ignore packages with persist event we need to compare head revision against entry in the persist table.
@@ -2595,7 +2546,7 @@ bool FConcertSyncSessionDatabase::EnumerateHeadRevisionPackageData(TFunctionRef<
 	return Statements->GetUniquePackageNameIdsForPackageEvents([this, &InCallback](int64 InPackageNameId)
 	{
 		int64 PackageHeadRevision = 0;
-		if (Statements->GetPackageHeadRevision(InPackageNameId, PackageHeadRevision))
+		if (Statements->GetUnmutedPackageHeadRevision(InPackageNameId, PackageHeadRevision))
 		{
 			FString DataFilename;
 			FConcertPackageInfo PackageInfo;
@@ -2686,7 +2637,7 @@ bool FConcertSyncSessionDatabase::GetPackageHeadRevision(const FName InPackageNa
 		return true;
 	}
 
-	return Statements->GetPackageHeadRevision(PackageNameId, OutRevision);
+	return Statements->GetUnmutedPackageHeadRevision(PackageNameId, OutRevision);
 }
 
 bool FConcertSyncSessionDatabase::IsHeadRevisionPackageEvent(const int64 InPackageEventId, bool& OutIsHeadRevision) const
@@ -2694,7 +2645,7 @@ bool FConcertSyncSessionDatabase::IsHeadRevisionPackageEvent(const int64 InPacka
 	int64 PackageNameId = 0;
 	int64 PackageRevision = 0;
 	int64 PackageHeadRevision = 0;
-	if (Statements->GetPackageNameIdAndRevisonForId(InPackageEventId, PackageNameId, PackageRevision) && Statements->GetPackageHeadRevision(PackageNameId, PackageHeadRevision))
+	if (Statements->GetPackageNameIdAndRevisonForId(InPackageEventId, PackageNameId, PackageRevision) && Statements->GetUnmutedPackageHeadRevision(PackageNameId, PackageHeadRevision))
 	{
 		OutIsHeadRevision = PackageRevision == PackageHeadRevision;
 		return true;
diff --git a/Engine/Plugins/Developer/Concert/ConcertSync/ConcertSyncCore/Source/ConcertSyncCore/Public/ConcertSyncSessionDatabase.h b/Engine/Plugins/Developer/Concert/ConcertSync/ConcertSyncCore/Source/ConcertSyncCore/Public/ConcertSyncSessionDatabase.h
--- a/Engine/Plugins/Developer/Concert/ConcertSync/ConcertSyncCore/Source/ConcertSyncCore/Public/ConcertSyncSessionDatabase.h
+++ b/Engine/Plugins/Developer/Concert/ConcertSync/ConcertSyncCore/Source/ConcertSyncCore/Public/ConcertSyncSessionDatabase.h
@@ -535,15 +535,6 @@ public:
 	 */
 	bool EnumerateLiveTransactionEventIdsForPackage(const FName InPackageName, TFunctionRef<bool(int64)> InCallback) const;
 
-	/**
-	 * Get the names of of any packages that have live transaction events.
-	 *
-	 * @param OutPackageNames			The array of names to populate with the result.
-	 *
-	 * @return True if the package names were resolved, false otherwise.
-	 */
-	bool GetPackageNamesWithLiveTransactions(TArray<FName>& OutPackageNames) const;
-
 	/**
 	 * Enumerate the names of of any packages that have live transaction events.
 	 *
@@ -585,16 +576,6 @@ public:
 	 */
 	bool GetPackageEvent(const int64 InPackageEventId, const TFunctionRef<void(FConcertSyncPackageEventData&)>& PackageEventFn) const;
 
-	/**
-	 * Get package names for packages with a head revision (at least one package event)
-	 *
-	 * @param OutPackageNames			The array of names to populate with the result.
-	 * @param IgnorePersisted			Will ignore packages which head revision have been persisted.
-	 *
-	 * @return True if the package names were resolved, false otherwise.
-	 */
-	bool GetPackageNamesWithHeadRevision(TArray<FName>& OutPackageNames, bool IgnorePersisted) const;
-
 	/**
 	 * Enumerate package names for packages with a head revision (at least one package event)
 	 *
