Preparing to allow Concert updates to reset existing objects to their archetype state

#rnx

diff --git a/Engine/Plugins/Developer/Concert/ConcertSync/ConcertSyncClient/Source/ConcertSyncClient/Private/ConcertSyncClientUtil.cpp b/Engine/Plugins/Developer/Concert/ConcertSync/ConcertSyncClient/Source/ConcertSyncClient/Private/ConcertSyncClientUtil.cpp
--- a/Engine/Plugins/Developer/Concert/ConcertSync/ConcertSyncClient/Source/ConcertSyncClient/Private/ConcertSyncClientUtil.cpp
+++ b/Engine/Plugins/Developer/Concert/ConcertSync/ConcertSyncClient/Source/ConcertSyncClient/Private/ConcertSyncClientUtil.cpp
@@ -343,16 +343,14 @@ TArray<FName> GetRootProperties(const TArray<FName>& InChangedProperties)
 
 FProperty* GetExportedProperty(const UStruct* InStruct, const FName InPropertyName, const bool InIncludeEditorOnlyData)
 {
-	FProperty* Property = FindFProperty<FProperty>(InStruct, InPropertyName);
-
-	// Filter the property
-	if (Property 
-		&& (!Property->IsEditorOnlyProperty() || InIncludeEditorOnlyData) 
-		&& !Property->HasAnyPropertyFlags(CPF_NonTransactional)
-		&& !ConcertSyncUtil::ShouldSkipTransientProperty(Property))
+	if (FProperty* Property = FindFProperty<FProperty>(InStruct, InPropertyName))
 	{
-		return Property;
+		if (ConcertSyncUtil::CanExportProperty(Property, InIncludeEditorOnlyData))
+		{
+			return Property;
+		}
 	}
+
 	return nullptr;
 }
 
diff --git a/Engine/Plugins/Developer/Concert/ConcertSync/ConcertSyncCore/Source/ConcertSyncCore/Private/ConcertSyncArchives.cpp b/Engine/Plugins/Developer/Concert/ConcertSync/ConcertSyncCore/Source/ConcertSyncCore/Private/ConcertSyncArchives.cpp
--- a/Engine/Plugins/Developer/Concert/ConcertSync/ConcertSyncCore/Source/ConcertSyncCore/Private/ConcertSyncArchives.cpp
+++ b/Engine/Plugins/Developer/Concert/ConcertSync/ConcertSyncCore/Source/ConcertSyncCore/Private/ConcertSyncArchives.cpp
@@ -8,6 +8,8 @@
 #include "Misc/PackageName.h"
 #include "Misc/EngineVersion.h"
 #include "Serialization/CustomVersion.h"
+#include "Serialization/ObjectReader.h"
+#include "Serialization/ObjectWriter.h"
 #include "UObject/Object.h"
 #include "UObject/Package.h"
 #include "UObject/LinkerLoad.h"
@@ -24,28 +26,149 @@ static const FName SkipAssetsMarker = TEXT("SKIPASSETS");
 
 namespace ConcertSyncUtil
 {
-	bool ShouldSkipTransientProperty(const FProperty* Property)
+
+bool CanExportProperty(const FProperty* Property, const bool InIncludeEditorOnlyData)
+{
+	auto CanExportTransientProperty = [Property]()
+	{
+		const UConcertSyncConfig* SyncConfig = GetDefault<UConcertSyncConfig>();
+		return SyncConfig->AllowedTransientProperties.ContainsByPredicate([Property](const TFieldPath<FProperty>& TransactionProperty)
+		{
+			FProperty* FilterProperty = TransactionProperty.Get();
+			return Property == FilterProperty;
+		});
+	};
+
+	return (!Property->IsEditorOnlyProperty() || InIncludeEditorOnlyData)
+		&& (!Property->HasAnyPropertyFlags(CPF_NonTransactional))
+		&& (!Property->HasAnyPropertyFlags(CPF_Transient) || CanExportTransientProperty());
+}
+
+void ResetObjectPropertiesToArchetypeValues(UObject* Object, const bool InIncludeEditorOnlyData)
+{
+	class FArchetypePropertyWriter : public FObjectWriter
+	{
+	public:
+		FArchetypePropertyWriter(const UObject* Obj, TArray<uint8>& InBytes, const bool InIncludeEditorOnlyData)
+			: FObjectWriter(InBytes)
+		{
+			ArIgnoreClassRef = true;
+			ArIgnoreArchetypeRef = true;
+			ArNoDelta = true;
+
+			SetFilterEditorOnly(!InIncludeEditorOnlyData);
+
+			Obj->SerializeScriptProperties(*this);
+		}
+
+		virtual bool ShouldSkipProperty(const FProperty* InProperty) const
+		{
+			return FObjectWriter::ShouldSkipProperty(InProperty)
+				|| InProperty->HasAnyPropertyFlags(CPF_DisableEditOnInstance);
+		}
+	};
+
+	class FArchetypePropertyReader : public FObjectReader
 	{
-		if (Property->HasAnyPropertyFlags(CPF_Transient))
+	public:
+		FArchetypePropertyReader(UObject* Obj, const TArray<uint8>& InBytes)
+			: FObjectReader(InBytes)
 		{
-			const UConcertSyncConfig* SyncConfig = GetDefault<UConcertSyncConfig>();
-			for (const TFieldPath<FProperty>& TransactionProperty : SyncConfig->AllowedTransientProperties)
+			ArIgnoreClassRef = true;
+			ArIgnoreArchetypeRef = true;
+
+#if USE_STABLE_LOCALIZATION_KEYS
+			if (GIsEditor && !(ArPortFlags & (PPF_DuplicateVerbatim | PPF_DuplicateForPIE)))
 			{
-				FProperty* FilterProperty = TransactionProperty.Get();
-				if (Property == FilterProperty)
-				{
-					// Allowed transient property, do not skip
-					return false;
-				}
+				SetLocalizationNamespace(TextNamespaceUtil::EnsurePackageNamespace(Obj));
+			}
+#endif // USE_STABLE_LOCALIZATION_KEYS
+
+			Obj->SerializeScriptProperties(*this);
+		}
+
+		virtual FArchive& operator<<(UObject*& Value) override
+		{
+			UObject* Tmp = nullptr;
+			FObjectReader::operator<<(Tmp);
+
+			if (!IsInstancedSubobject(Value))
+			{
+				Value = Tmp;
 			}
-			// Not allowed transient property, skip
-			return true;
+
+			return *this;
 		}
-		// Non transient property, might not skip
-		return false;
+
+		virtual FArchive& operator<<(FObjectPtr& Value) override
+		{
+			FObjectPtr Tmp = nullptr;
+			FObjectReader::operator<<(Tmp);
+
+			if (!IsInstancedSubobject(Value.Get()))
+			{
+				Value = Tmp;
+			}
+
+			return *this;
+		}
+
+		virtual FArchive& operator<<(FLazyObjectPtr& Value) override
+		{
+			FLazyObjectPtr Tmp;
+			FObjectReader::operator<<(Tmp);
+
+			if (!IsInstancedSubobject(Value.Get()))
+			{
+				Value = Tmp;
+			}
+
+			return *this;
+		}
+
+		virtual FArchive& operator<<(FSoftObjectPtr& Value) override
+		{
+			FSoftObjectPtr Tmp;
+			FObjectReader::operator<<(Tmp);
+
+			if (!IsInstancedSubobject(Value.Get()))
+			{
+				Value = Tmp;
+			}
+
+			return *this;
+		}
+
+		virtual FArchive& operator<<(FWeakObjectPtr& Value) override
+		{
+			FWeakObjectPtr Tmp;
+			FObjectReader::operator<<(Tmp);
+
+			if (!IsInstancedSubobject(Value.Get()))
+			{
+				Value = Tmp;
+			}
+
+			return *this;
+		}
+
+	private:
+		bool IsInstancedSubobject(const UObject* Obj) const
+		{
+			return Obj && (Obj->HasAnyFlags(RF_DefaultSubObject) || Obj->IsDefaultSubobject());
+		}
+	};
+
+	if (const UObject* ObjectArchetype = Object->GetArchetype(); ObjectArchetype && !Object->HasAnyFlags(RF_ClassDefaultObject | RF_ArchetypeObject))
+	{
+		TArray<uint8> ArchetypeData;
+		FArchetypePropertyWriter(ObjectArchetype, ArchetypeData, InIncludeEditorOnlyData);
+		FArchetypePropertyReader(Object, ArchetypeData);
 	}
 }
 
+} // namespace ConcertSyncUtil
+
 FString FConcertSyncWorldRemapper::RemapObjectPathName(const FString& InObjectPathName) const
 {
 	if (RemapDelegate.IsBound())
@@ -202,8 +325,8 @@ FString FConcertSyncObjectWriter::GetArchiveName() const
 
 bool FConcertSyncObjectWriter::ShouldSkipProperty(const FProperty* InProperty) const
 {
-	return (ShouldSkipPropertyFunc && ShouldSkipPropertyFunc(InProperty)) || 
-		ConcertSyncUtil::ShouldSkipTransientProperty(InProperty);
+	return (ShouldSkipPropertyFunc && ShouldSkipPropertyFunc(InProperty))
+		|| (!ConcertSyncUtil::CanExportProperty(InProperty, !IsFilterEditorOnly()));
 }
 
 FConcertSyncObjectReader::FConcertSyncObjectReader(const FConcertLocalIdentifierTable* InLocalIdentifierTable, FConcertSyncWorldRemapper InWorldRemapper, const FConcertSessionVersionInfo* InVersionInfo, UObject* InObj, const TArray<uint8>& InBytes)
diff --git a/Engine/Plugins/Developer/Concert/ConcertSync/ConcertSyncCore/Source/ConcertSyncCore/Public/ConcertSyncArchives.h b/Engine/Plugins/Developer/Concert/ConcertSync/ConcertSyncCore/Source/ConcertSyncCore/Public/ConcertSyncArchives.h
--- a/Engine/Plugins/Developer/Concert/ConcertSync/ConcertSyncCore/Source/ConcertSyncCore/Public/ConcertSyncArchives.h
+++ b/Engine/Plugins/Developer/Concert/ConcertSync/ConcertSyncCore/Source/ConcertSyncCore/Public/ConcertSyncArchives.h
@@ -12,7 +12,9 @@ DECLARE_DELEGATE_RetVal_OneParam(bool, FConcertSyncObjectPathBelongsToWorld, con
 
 namespace ConcertSyncUtil
 {
-	 bool CONCERTSYNCCORE_API ShouldSkipTransientProperty(const FProperty* Property);
+CONCERTSYNCCORE_API bool CanExportProperty(const FProperty* Property, const bool InIncludeEditorOnlyData);
+
+CONCERTSYNCCORE_API void ResetObjectPropertiesToArchetypeValues(UObject* Object, const bool InIncludeEditorOnlyData);
 }
 
 /** Util to handle remapping objects within the source world to be the equivalent objects in the destination world */
diff --git a/Engine/Plugins/Developer/Concert/ConcertSync/ConcertSyncCore/Source/ConcertSyncCore/Public/ConcertTransactionEvents.h b/Engine/Plugins/Developer/Concert/ConcertSync/ConcertSyncCore/Source/ConcertSyncCore/Public/ConcertTransactionEvents.h
--- a/Engine/Plugins/Developer/Concert/ConcertSync/ConcertSyncCore/Source/ConcertSyncCore/Public/ConcertTransactionEvents.h
+++ b/Engine/Plugins/Developer/Concert/ConcertSync/ConcertSyncCore/Source/ConcertSyncCore/Public/ConcertTransactionEvents.h
@@ -79,6 +79,9 @@ struct FConcertSerializedObjectData
 	UPROPERTY()
 	bool bAllowCreate = false;
 
+	UPROPERTY()
+	bool bResetExisting = false;
+
 	UPROPERTY()
 	bool bIsPendingKill = false;
 
