Only reset editable properties when attempting to restore the archetype state in Concert

Non-editable properties may have been set via other methods (like BeginPlay) and resetting them can leave an object in an invalid state.

#rnx

diff --git a/Engine/Plugins/Developer/Concert/ConcertSync/ConcertSyncCore/Source/ConcertSyncCore/Private/ConcertSyncArchives.cpp b/Engine/Plugins/Developer/Concert/ConcertSync/ConcertSyncCore/Source/ConcertSyncCore/Private/ConcertSyncArchives.cpp
--- a/Engine/Plugins/Developer/Concert/ConcertSync/ConcertSyncCore/Source/ConcertSyncCore/Private/ConcertSyncArchives.cpp
+++ b/Engine/Plugins/Developer/Concert/ConcertSync/ConcertSyncCore/Source/ConcertSyncCore/Private/ConcertSyncArchives.cpp
@@ -64,7 +64,8 @@ void ResetObjectPropertiesToArchetypeValues(UObject* Object, const bool InInclud
 		virtual bool ShouldSkipProperty(const FProperty* InProperty) const
 		{
 			return FObjectWriter::ShouldSkipProperty(InProperty)
-				|| InProperty->HasAnyPropertyFlags(CPF_DisableEditOnInstance);
+				|| !InProperty->HasAnyPropertyFlags(CPF_Edit)
+				|| InProperty->HasAnyPropertyFlags(CPF_EditConst | CPF_DisableEditOnInstance);
 		}
 	};
 
