Allow transient AttachChildren property to replicate via Concert

#rnx

diff --git a/Engine/Plugins/Developer/Concert/ConcertSync/ConcertSyncCore/Config/BaseConcertSyncCore.ini b/Engine/Plugins/Developer/Concert/ConcertSync/ConcertSyncCore/Config/BaseConcertSyncCore.ini
--- a/Engine/Plugins/Developer/Concert/ConcertSync/ConcertSyncCore/Config/BaseConcertSyncCore.ini
+++ b/Engine/Plugins/Developer/Concert/ConcertSync/ConcertSyncCore/Config/BaseConcertSyncCore.ini
@@ -7,6 +7,7 @@
 +AllowedTransientProperties="/Script/Engine.Actor:bHiddenEdTemporary"
 +AllowedTransientProperties="/Script/Engine.Actor:bHiddenEdLevel"
 +AllowedTransientProperties="/Script/Engine.Level:bIsVisible"
++AllowedTransientProperties="/Script/Engine.SceneComponent:AttachChildren"
 
 ; If you wish to use Composure in Multi-user, you may add that line under [/Script/ConcertSyncCore.ConcertSyncConfig] in your project engine config settings.
 ;+IncludeObjectClassFilters=(ObjectOuterClass=None,ObjectClasses=(Class'/Script/Composure.CompositingElementPass',Class'/Script/Composure.ComposurePostProcessPassPolicy'))
