Fix uninited variable

diff --git a/Engine/Plugins/Developer/Concert/ConcertSync/ConcertSyncCore/Source/ConcertSyncCore/Public/ConcertWorkspaceMessages.h b/Engine/Plugins/Developer/Concert/ConcertSync/ConcertSyncCore/Source/ConcertSyncCore/Public/ConcertWorkspaceMessages.h
--- a/Engine/Plugins/Developer/Concert/ConcertSync/ConcertSyncCore/Source/ConcertSyncCore/Public/ConcertWorkspaceMessages.h
+++ b/Engine/Plugins/Developer/Concert/ConcertSync/ConcertSyncCore/Source/ConcertSyncCore/Public/ConcertWorkspaceMessages.h
@@ -121,7 +121,7 @@ struct FConcertPackageTransmissionStartEvent
 	FConcertPackageInfo PackageInfo;
 
 	UPROPERTY()
-	uint64 PackageNumBytes;
+	uint64 PackageNumBytes = 0;
 };
 
 
