Improved logic for resetting an object to its archetype state when reusing an existing instance in Concert

This now more aggressively resets properties (as there are non-editable properties, such as bActorIsBeingDestroyed) that need to reset when reusing an actor that was destroyed.

This also fixes a bug where a null pointer could be overwritten by a DSO from the archetype.

#rnx

diff --git a/Engine/Plugins/Developer/Concert/ConcertSync/ConcertSyncCore/Source/ConcertSyncCore/Private/ConcertSyncArchives.cpp b/Engine/Plugins/Developer/Concert/ConcertSync/ConcertSyncCore/Source/ConcertSyncCore/Private/ConcertSyncArchives.cpp
--- a/Engine/Plugins/Developer/Concert/ConcertSync/ConcertSyncCore/Source/ConcertSyncCore/Private/ConcertSyncArchives.cpp
+++ b/Engine/Plugins/Developer/Concert/ConcertSync/ConcertSyncCore/Source/ConcertSyncCore/Private/ConcertSyncArchives.cpp
@@ -44,13 +44,24 @@ bool CanExportProperty(const FProperty* Property, const bool InIncludeEditorOnly
 		&& (!Property->HasAnyPropertyFlags(CPF_Transient) || CanExportTransientProperty());
 }
 
+void GatherDefaultSubobjectPaths(const UObject* Obj, TSet<FSoftObjectPath>& OutSubobjects)
+{
+	ForEachObjectWithOuter(Obj, [&OutSubobjects](UObject* InnerObj)
+	{
+		if (InnerObj->HasAnyFlags(RF_DefaultSubObject) || InnerObj->IsDefaultSubobject())
+		{
+			OutSubobjects.Emplace(InnerObj);
+		}
+	});
+}
+
 void ResetObjectPropertiesToArchetypeValues(UObject* Object, const bool InIncludeEditorOnlyData)
 {
 	class FArchetypePropertyWriter : public FObjectWriter
 	{
 	public:
-		FArchetypePropertyWriter(const UObject* Obj, TArray<uint8>& InBytes, const bool InIncludeEditorOnlyData)
-			: FObjectWriter(InBytes)
+		FArchetypePropertyWriter(const UObject* Obj, TArray<uint8>& OutBytes, TSet<FSoftObjectPath>& OutObjectsToSkip, const bool InIncludeEditorOnlyData)
+			: FObjectWriter(OutBytes)
 		{
 			ArIgnoreClassRef = true;
 			ArIgnoreArchetypeRef = true;
@@ -58,26 +69,23 @@ void ResetObjectPropertiesToArchetypeValues(UObject* Object, const bool InInclud
 
 			SetFilterEditorOnly(!InIncludeEditorOnlyData);
 
-			Obj->SerializeScriptProperties(*this);
-		}
+			GatherDefaultSubobjectPaths(Obj, OutObjectsToSkip);
 
-		virtual bool ShouldSkipProperty(const FProperty* InProperty) const
-		{
-			return FObjectWriter::ShouldSkipProperty(InProperty)
-				|| !InProperty->HasAnyPropertyFlags(CPF_Edit)
-				|| InProperty->HasAnyPropertyFlags(CPF_EditConst | CPF_DisableEditOnInstance);
+			Obj->SerializeScriptProperties(*this);
 		}
 	};
 
 	class FArchetypePropertyReader : public FObjectReader
 	{
 	public:
-		FArchetypePropertyReader(UObject* Obj, const TArray<uint8>& InBytes)
+		FArchetypePropertyReader(UObject* Obj, const TArray<uint8>& InBytes, const TSet<FSoftObjectPath>& InObjectsToSkip, const bool InIncludeEditorOnlyData)
 			: FObjectReader(InBytes)
 		{
 			ArIgnoreClassRef = true;
 			ArIgnoreArchetypeRef = true;
 
+			SetFilterEditorOnly(!InIncludeEditorOnlyData);
+
 #if USE_STABLE_LOCALIZATION_KEYS
 			if (GIsEditor && !(ArPortFlags & (PPF_DuplicateVerbatim | PPF_DuplicateForPIE)))
 			{
@@ -85,6 +93,9 @@ void ResetObjectPropertiesToArchetypeValues(UObject* Object, const bool InInclud
 			}
 #endif // USE_STABLE_LOCALIZATION_KEYS
 
+			ObjectsToSkip = InObjectsToSkip;
+			GatherDefaultSubobjectPaths(Obj, ObjectsToSkip);
+
 			Obj->SerializeScriptProperties(*this);
 		}
 
@@ -93,7 +104,7 @@ void ResetObjectPropertiesToArchetypeValues(UObject* Object, const bool InInclud
 			UObject* Tmp = nullptr;
 			FObjectReader::operator<<(Tmp);
 
-			if (!IsInstancedSubobject(Value))
+			if (CanOverwriteObject(Value, Tmp))
 			{
 				Value = Tmp;
 			}
@@ -106,7 +117,7 @@ void ResetObjectPropertiesToArchetypeValues(UObject* Object, const bool InInclud
 			FObjectPtr Tmp = nullptr;
 			FObjectReader::operator<<(Tmp);
 
-			if (!IsInstancedSubobject(Value.Get()))
+			if (CanOverwriteObject(Value.Get(), Tmp.Get()))
 			{
 				Value = Tmp;
 			}
@@ -119,7 +130,7 @@ void ResetObjectPropertiesToArchetypeValues(UObject* Object, const bool InInclud
 			FLazyObjectPtr Tmp;
 			FObjectReader::operator<<(Tmp);
 
-			if (!IsInstancedSubobject(Value.Get()))
+			if (CanOverwriteObject(Value.Get(), Tmp.Get()))
 			{
 				Value = Tmp;
 			}
@@ -132,7 +143,20 @@ void ResetObjectPropertiesToArchetypeValues(UObject* Object, const bool InInclud
 			FSoftObjectPtr Tmp;
 			FObjectReader::operator<<(Tmp);
 
-			if (!IsInstancedSubobject(Value.Get()))
+			if (CanOverwriteObject(Value.ToSoftObjectPath(), Tmp.ToSoftObjectPath()))
+			{
+				Value = Tmp;
+			}
+
+			return *this;
+		}
+
+		virtual FArchive& operator<<(FSoftObjectPath& Value) override
+		{
+			FSoftObjectPath Tmp;
+			FObjectReader::operator<<(Tmp);
+
+			if (CanOverwriteObject(Value, Tmp))
 			{
 				Value = Tmp;
 			}
@@ -145,7 +169,7 @@ void ResetObjectPropertiesToArchetypeValues(UObject* Object, const bool InInclud
 			FWeakObjectPtr Tmp;
 			FObjectReader::operator<<(Tmp);
 
-			if (!IsInstancedSubobject(Value.Get()))
+			if (CanOverwriteObject(Value.Get(), Tmp.Get()))
 			{
 				Value = Tmp;
 			}
@@ -154,17 +178,20 @@ void ResetObjectPropertiesToArchetypeValues(UObject* Object, const bool InInclud
 		}
 
 	private:
-		bool IsInstancedSubobject(const UObject* Obj) const
+		bool CanOverwriteObject(const FSoftObjectPath& CurObj, const FSoftObjectPath& NewObj) const
 		{
-			return Obj && (Obj->HasAnyFlags(RF_DefaultSubObject) || Obj->IsDefaultSubobject());
+			return !ObjectsToSkip.Contains(CurObj) && !ObjectsToSkip.Contains(NewObj);
 		}
+
+		TSet<FSoftObjectPath> ObjectsToSkip;
 	};
 
 	if (const UObject* ObjectArchetype = Object->GetArchetype())
 	{
 		TArray<uint8> ArchetypeData;
-		FArchetypePropertyWriter(ObjectArchetype, ArchetypeData, InIncludeEditorOnlyData);
-		FArchetypePropertyReader(Object, ArchetypeData);
+		TSet<FSoftObjectPath> ArchetypeObjectsToSkip;
+		FArchetypePropertyWriter(ObjectArchetype, ArchetypeData, ArchetypeObjectsToSkip, InIncludeEditorOnlyData);
+		FArchetypePropertyReader(Object, ArchetypeData, ArchetypeObjectsToSkip, InIncludeEditorOnlyData);
 	}
 }
 
