Concert Plugin - Detect missing assets when reading messages

diff --git a/Engine/Plugins/Developer/Concert/ConcertSync/ConcertSyncCore/Source/ConcertSyncCore/Private/ConcertSyncArchives.cpp b/Engine/Plugins/Developer/Concert/ConcertSync/ConcertSyncCore/Source/ConcertSyncCore/Private/ConcertSyncArchives.cpp
--- a/Engine/Plugins/Developer/Concert/ConcertSync/ConcertSyncCore/Source/ConcertSyncCore/Private/ConcertSyncArchives.cpp
+++ b/Engine/Plugins/Developer/Concert/ConcertSync/ConcertSyncCore/Source/ConcertSyncCore/Private/ConcertSyncArchives.cpp
@@ -357,9 +357,10 @@ bool FConcertSyncObjectWriter::ShouldSkipProperty(const FProperty* InProperty) c
 		|| (!ConcertSyncUtil::CanExportProperty(InProperty, !IsFilterEditorOnly()));
 }
 
-FConcertSyncObjectReader::FConcertSyncObjectReader(const FConcertLocalIdentifierTable* InLocalIdentifierTable, FConcertSyncWorldRemapper InWorldRemapper, const FConcertSessionVersionInfo* InVersionInfo, UObject* InObj, const TArray<uint8>& InBytes)
+FConcertSyncObjectReader::FConcertSyncObjectReader(const FConcertLocalIdentifierTable* InLocalIdentifierTable, FConcertSyncWorldRemapper InWorldRemapper, const FConcertSessionVersionInfo* InVersionInfo, UObject* InObj, const TArray<uint8>& InBytes, const FConcertSyncEncounteredMissingObject& InEncounteredMissingObjectDelegate)
 	: FConcertIdentifierReader(InLocalIdentifierTable, InBytes, /*bIsPersistent*/false)
 	, WorldRemapper(MoveTemp(InWorldRemapper))
+	, EncounteredMissingObjectDelegate(InEncounteredMissingObjectDelegate)
 {
 	ArIgnoreClassRef = false;
 	ArIgnoreArchetypeRef = false;
@@ -395,6 +396,11 @@ FConcertSyncObjectReader::FConcertSyncObjectReader(const FConcertLocalIdentifier
 #endif // USE_STABLE_LOCALIZATION_KEYS
 }
 
+FConcertSyncObjectReader::FConcertSyncObjectReader(const FConcertLocalIdentifierTable* InLocalIdentifierTable, FConcertSyncWorldRemapper InWorldRemapper, const FConcertSessionVersionInfo* InVersionInfo, UObject* InObj, const TArray<uint8>& InBytes)
+	: FConcertSyncObjectReader(InLocalIdentifierTable, InWorldRemapper, InVersionInfo, InObj, InBytes, FConcertSyncEncounteredMissingObject())
+{
+}
+
 void FConcertSyncObjectReader::SerializeObject(UObject* InObject)
 {
 	InObject->Serialize(*this);
@@ -442,6 +448,11 @@ FArchive& FConcertSyncObjectReader::operator<<(UObject*& Obj)
 					Obj = StaticLoadObject(UObject::StaticClass(), nullptr, *ResolvedObjPath);
 				}
 			}
+
+			if (!Obj && EncounteredMissingObjectDelegate.IsBound())
+			{
+				EncounteredMissingObjectDelegate.Execute(FStringView(ResolvedObjPath));
+			}
 		}
 	}
 
diff --git a/Engine/Plugins/Developer/Concert/ConcertSync/ConcertSyncCore/Source/ConcertSyncCore/Public/ConcertSyncArchives.h b/Engine/Plugins/Developer/Concert/ConcertSync/ConcertSyncCore/Source/ConcertSyncCore/Public/ConcertSyncArchives.h
--- a/Engine/Plugins/Developer/Concert/ConcertSync/ConcertSyncCore/Source/ConcertSyncCore/Public/ConcertSyncArchives.h
+++ b/Engine/Plugins/Developer/Concert/ConcertSync/ConcertSyncCore/Source/ConcertSyncCore/Public/ConcertSyncArchives.h
@@ -9,6 +9,7 @@ struct FConcertSessionVersionInfo;
 
 DECLARE_DELEGATE_OneParam(FConcertSyncRemapObjectPath, FString& /*ObjectPath*/)
 DECLARE_DELEGATE_RetVal_OneParam(bool, FConcertSyncObjectPathBelongsToWorld, const FStringView /*ObjectPath*/)
+DECLARE_DELEGATE_OneParam(FConcertSyncEncounteredMissingObject, const FStringView /*ObjectPath*/)
 
 namespace ConcertSyncUtil
 {
@@ -77,6 +78,7 @@ private:
 class CONCERTSYNCCORE_API FConcertSyncObjectReader : public FConcertIdentifierReader
 {
 public:
+	FConcertSyncObjectReader(const FConcertLocalIdentifierTable* InLocalIdentifierTable, FConcertSyncWorldRemapper InWorldRemapper, const FConcertSessionVersionInfo* InVersionInfo, UObject* InObj, const TArray<uint8>& InBytes, const FConcertSyncEncounteredMissingObject& InEncounteredMissingObjectDelegate);
 	FConcertSyncObjectReader(const FConcertLocalIdentifierTable* InLocalIdentifierTable, FConcertSyncWorldRemapper InWorldRemapper, const FConcertSessionVersionInfo* InVersionInfo, UObject* InObj, const TArray<uint8>& InBytes);
 
 	void SerializeObject(UObject* InObject);
@@ -96,4 +98,5 @@ public:
 
 private:
 	FConcertSyncWorldRemapper WorldRemapper;
+	FConcertSyncEncounteredMissingObject EncounteredMissingObjectDelegate;
 };
