Add ability to persist a single file in blueprints. Also added a delegate that indicates when a new package has been saved.

diff --git a/Engine/Plugins/Developer/Concert/ConcertApp/MultiUserClient/Source/MultiUserClientLibrary/MultiUserClientLibrary.Build.cs b/Engine/Plugins/Developer/Concert/ConcertApp/MultiUserClient/Source/MultiUserClientLibrary/MultiUserClientLibrary.Build.cs
--- a/Engine/Plugins/Developer/Concert/ConcertApp/MultiUserClient/Source/MultiUserClientLibrary/MultiUserClientLibrary.Build.cs
+++ b/Engine/Plugins/Developer/Concert/ConcertApp/MultiUserClient/Source/MultiUserClientLibrary/MultiUserClientLibrary.Build.cs
@@ -19,14 +19,15 @@ namespace UnrealBuildTool.Rules
 			{
 				PrivateDefinitions.Add("WITH_CONCERT=1");
 
-                PrivateDependencyModuleNames.AddRange(
-                    new string[]
-                    {
-                        "Concert"
-                    }
-                );
+				PrivateDependencyModuleNames.AddRange(
+					new string[]
+					{
+						"Concert",
+						"ConcertSyncCore"
+					}
+				);
 
-                PrivateIncludePathModuleNames.AddRange(
+				PrivateIncludePathModuleNames.AddRange(
 					new string[]
 					{
 						"ConcertSyncCore",
diff --git a/Engine/Plugins/Developer/Concert/ConcertApp/MultiUserClient/Source/MultiUserClientLibrary/Private/MultiUserClientStatics.cpp b/Engine/Plugins/Developer/Concert/ConcertApp/MultiUserClient/Source/MultiUserClientLibrary/Private/MultiUserClientStatics.cpp
--- a/Engine/Plugins/Developer/Concert/ConcertApp/MultiUserClient/Source/MultiUserClientLibrary/Private/MultiUserClientStatics.cpp
+++ b/Engine/Plugins/Developer/Concert/ConcertApp/MultiUserClient/Source/MultiUserClientLibrary/Private/MultiUserClientStatics.cpp
@@ -11,6 +11,8 @@
 #include "IConcertClientPresenceManager.h"
 #include "ConcertSettings.h"
 #include "ConcertMessageData.h"
+#include "ConcertSyncSessionDatabase.h"
+#include "IConcertClientWorkspace.h"
 #endif
 
 DEFINE_LOG_CATEGORY_STATIC(LogMultiUserClient, Log, All);
@@ -172,6 +174,63 @@ void UMultiUserClientStatics::PersistMultiUserSessionChanges()
 #endif
 }
 
+void UMultiUserClientStatics::PersistSpecifiedPackages(const TArray<FName>& PackagesToPersist)
+{
+#if WITH_CONCERT
+	if (IMultiUserClientModule::IsAvailable())
+	{
+		if (TSharedPtr<IConcertSyncClient> ConcertSyncClient = IMultiUserClientModule::Get().GetClient())
+		{
+			if (TSharedPtr<IConcertClientWorkspace> Workspace = ConcertSyncClient->GetWorkspace())
+			{
+				ConcertSyncClient->PersistSpecificChanges({PackagesToPersist});
+			}
+		}
+	}
+#endif
+}
+
+TArray<FName> UMultiUserClientStatics::GatherSessionChanges(bool bIgnorePersisted)
+{
+#if WITH_CONCERT
+	if (IMultiUserClientModule::IsAvailable())
+	{
+		if (TSharedPtr<IConcertSyncClient> ConcertSyncClient = IMultiUserClientModule::Get().GetClient())
+		{
+			if (TSharedPtr<IConcertClientWorkspace> Workspace = ConcertSyncClient->GetWorkspace())
+			{
+				return Workspace->GatherSessionChanges();
+			}
+		}
+	}
+#endif
+	return {};
+}
+
+UMultiUserClientSyncDatabase* UMultiUserClientStatics::GetConcertSyncDatabase()
+{
+#if WITH_CONCERT
+	static bool bRunOnce = []()
+	{
+		UE::ConcertSyncCore::SyncDatabase::GetOnPackageSavedDelegate().AddLambda([](const FName& PackageName)
+		{
+			UMultiUserClientSyncDatabase* Database = GetMutableDefault<UMultiUserClientSyncDatabase>();
+			check(Database);
+			// Re-broadcast message to blueprints.
+			Database->OnPackageSaved.Broadcast(PackageName);
+		});
+		return true;
+	}();
+
+	UMultiUserClientSyncDatabase* SyncDatabase = GetMutableDefault<UMultiUserClientSyncDatabase>();
+	if (bRunOnce && SyncDatabase)
+	{
+		return SyncDatabase;
+	}
+#endif
+	return nullptr;
+}
+
 FMultiUserClientInfo UMultiUserClientStatics::GetLocalMultiUserClientInfo()
 {
 	FMultiUserClientInfo ClientInfo;
@@ -182,7 +241,7 @@ FMultiUserClientInfo UMultiUserClientStatics::GetLocalMultiUserClientInfo()
 		{
 			IConcertClientRef ConcertClient = ConcertSyncClient->GetConcertClient();
 			TSharedPtr<IConcertClientSession> ClientSession = ConcertClient->GetCurrentSession();
-		
+
 			FGuid LocalClientEndpointId = ClientSession ? ClientSession->GetSessionClientEndpointId() : FGuid();
 			const FConcertClientInfo& LocalClientInfo = ClientSession ? ClientSession->GetLocalClientInfo() : ConcertClient->GetClientInfo();
 			ClientInfo = MultiUserClientUtil::ConvertClientInfo(LocalClientEndpointId, LocalClientInfo);
diff --git a/Engine/Plugins/Developer/Concert/ConcertApp/MultiUserClient/Source/MultiUserClientLibrary/Public/MultiUserClientStatics.h b/Engine/Plugins/Developer/Concert/ConcertApp/MultiUserClient/Source/MultiUserClientLibrary/Public/MultiUserClientStatics.h
--- a/Engine/Plugins/Developer/Concert/ConcertApp/MultiUserClient/Source/MultiUserClientLibrary/Public/MultiUserClientStatics.h
+++ b/Engine/Plugins/Developer/Concert/ConcertApp/MultiUserClient/Source/MultiUserClientLibrary/Public/MultiUserClientStatics.h
@@ -6,6 +6,9 @@
 #include "Kismet/BlueprintFunctionLibrary.h"
 #include "MultiUserClientStatics.generated.h"
 
+/** Delegate that is invoked when a package is saved. */
+DECLARE_DYNAMIC_MULTICAST_DELEGATE_OneParam(FOnPackageSavedSignature, FName, PackageName);
+
 /**
  * BP copy of FConcertSessionClientInfo
  * Holds info on a client connected through multi-user
@@ -32,6 +35,16 @@ struct FMultiUserClientInfo
 	TArray<FName> Tags;
 };
 
+UCLASS()
+class UMultiUserClientSyncDatabase : public UObject
+{
+	GENERATED_BODY()
+public:
+
+	UPROPERTY(VisibleAnywhere, BlueprintAssignable, Category="Save Package Delegate")
+	FOnPackageSavedSignature OnPackageSaved;
+};
+
 /**
  * Enum of the known Multi-User connection error, their value needs to match the internally returned error code.
  * @see FConcertConnectionError
@@ -145,6 +158,18 @@ public:
 	UFUNCTION(BlueprintCallable, Category = "Multi-User Source Control", meta=(DevelopmentOnly, DisplayName = "Persist Multi-User Session Changes"))
 	static void PersistMultiUserSessionChanges();
 
+	/** Persist the specified sessions changes using source control. */
+	UFUNCTION(BlueprintCallable, Category = "Multi-User Source Control", meta=(DevelopmentOnly, DsiplayName = "Persist Multi-User Sessions Changes (By Package name)"))
+	static void PersistSpecifiedPackages(const TArray<FName>& PackagesToPersist);
+
+	/** Get the list of packages that have changed since the last persist unless ignore persisted is false.  */
+	UFUNCTION(BlueprintCallable, Category = "Multi-User Source Control", meta=(DevelopmentOnly, DsiplayName = "Get List of Packages that have changed."))
+	static TArray<FName> GatherSessionChanges(bool IgnorePersisted = true);
+
+	/** Get the proxy object for the sync database. */
+	UFUNCTION(BlueprintCallable, Category = "Multi-User Source Control", meta=(DevelopmentOnly, DsiplayName = "Get Multi-user Sync Database."))
+	static UMultiUserClientSyncDatabase* GetConcertSyncDatabase();
+
 	/** Get the local ClientInfo. Works when not connected to a session. */
 	UFUNCTION(BlueprintCallable, Category = "Multi-User Client", meta=(DevelopmentOnly, DisplayName = "Get Local Multi-User Client Info"))
 	static FMultiUserClientInfo GetLocalMultiUserClientInfo();
diff --git a/Engine/Plugins/Developer/Concert/ConcertSync/ConcertSyncClient/Source/ConcertSyncClient/Private/ConcertSyncClient.cpp b/Engine/Plugins/Developer/Concert/ConcertSync/ConcertSyncClient/Source/ConcertSyncClient/Private/ConcertSyncClient.cpp
--- a/Engine/Plugins/Developer/Concert/ConcertSync/ConcertSyncClient/Source/ConcertSyncClient/Private/ConcertSyncClient.cpp
+++ b/Engine/Plugins/Developer/Concert/ConcertSync/ConcertSyncClient/Source/ConcertSyncClient/Private/ConcertSyncClient.cpp
@@ -106,6 +106,16 @@ FOnConcertClientSyncSessionStartupOrShutdown& FConcertSyncClient::OnSyncSessionS
 	return OnSyncSessionShutdownDelegate;
 }
 
+void FConcertSyncClient::PersistSpecificChanges(TArrayView<const FName> InPackages)
+{
+#if WITH_EDITOR
+	if (Workspace)
+	{
+		Workspace->PersistSessionChanges({InPackages, SourceControlProxy.Get()});
+	}
+#endif
+}
+
 void FConcertSyncClient::PersistAllSessionChanges()
 {
 #if WITH_EDITOR
diff --git a/Engine/Plugins/Developer/Concert/ConcertSync/ConcertSyncClient/Source/ConcertSyncClient/Private/ConcertSyncClient.h b/Engine/Plugins/Developer/Concert/ConcertSync/ConcertSyncClient/Source/ConcertSyncClient/Private/ConcertSyncClient.h
--- a/Engine/Plugins/Developer/Concert/ConcertSync/ConcertSyncClient/Source/ConcertSyncClient/Private/ConcertSyncClient.h
+++ b/Engine/Plugins/Developer/Concert/ConcertSync/ConcertSyncClient/Source/ConcertSyncClient/Private/ConcertSyncClient.h
@@ -36,6 +36,7 @@ public:
 	virtual FOnConcertClientSyncSessionStartupOrShutdown& OnSyncSessionStartup() override;
 	virtual FOnConcertClientSyncSessionStartupOrShutdown& OnSyncSessionShutdown() override;
 
+	virtual void PersistSpecificChanges(TArrayView<const FName> InPackages) override;
 	virtual void PersistAllSessionChanges() override;
 	virtual void GetSessionClientActions(const FConcertSessionClientInfo& InClientInfo, TArray<FConcertActionDefinition>& OutActions) const override;
 
diff --git a/Engine/Plugins/Developer/Concert/ConcertSync/ConcertSyncClient/Source/ConcertSyncClient/Public/IConcertSyncClient.h b/Engine/Plugins/Developer/Concert/ConcertSync/ConcertSyncClient/Source/ConcertSyncClient/Public/IConcertSyncClient.h
--- a/Engine/Plugins/Developer/Concert/ConcertSync/ConcertSyncClient/Source/ConcertSyncClient/Public/IConcertSyncClient.h
+++ b/Engine/Plugins/Developer/Concert/ConcertSync/ConcertSyncClient/Source/ConcertSyncClient/Public/IConcertSyncClient.h
@@ -77,6 +77,9 @@ public:
 	/** Get the delegate called just after the workspace gets deleted. */
 	virtual FOnConcertClientSyncSessionStartupOrShutdown& OnSyncSessionShutdown() = 0;
 
+	/** Persist specific packages */
+	virtual void PersistSpecificChanges(TArrayView<const FName> Packages) = 0;
+
 	/** Persist all session changes and prepare the files for source control submission. */
 	virtual void PersistAllSessionChanges() = 0;
 
diff --git a/Engine/Plugins/Developer/Concert/ConcertSync/ConcertSyncCore/Source/ConcertSyncCore/Private/ConcertSyncSessionDatabase.cpp b/Engine/Plugins/Developer/Concert/ConcertSync/ConcertSyncCore/Source/ConcertSyncCore/Private/ConcertSyncSessionDatabase.cpp
--- a/Engine/Plugins/Developer/Concert/ConcertSync/ConcertSyncCore/Source/ConcertSyncCore/Private/ConcertSyncSessionDatabase.cpp
+++ b/Engine/Plugins/Developer/Concert/ConcertSync/ConcertSyncCore/Source/ConcertSyncCore/Private/ConcertSyncSessionDatabase.cpp
@@ -18,7 +18,6 @@
 #include "Serialization/MemoryWriter.h"
 #include "Async/Async.h"
 
-
 /** Contains an asynchronous representation of the FConcertPackageDataStream */
 namespace PackageDataUtil
 {
@@ -31,7 +30,7 @@ namespace PackageDataUtil
 
 struct FConcertPackageAsyncDataStream
 {
-	FConcertPackageAsyncDataStream(const FConcertPackageDataStream& InStream, const FString& InPackageName)
+	FConcertPackageAsyncDataStream(const FConcertPackageDataStream& InStream, const FName& InPackageName)
 		:
 		CachedPackageName(InPackageName),
 		PackageData(InStream.DataBlob->GetData(),InStream.DataBlob->Num())
@@ -39,7 +38,7 @@ struct FConcertPackageAsyncDataStream
 		PackageStream = {nullptr, InStream.DataSize, &PackageData};
 	}
 
-	FString								CachedPackageName;
+	FName								CachedPackageName;
 	TFuture<bool>						AsyncTask;
 	PackageDataUtil::WritePackageResult Result;
 
@@ -70,6 +69,12 @@ struct FDeferredLargePackageIOImpl
 	FAsyncDataStreamMap DeferredLargePackageIO;
 };
 
+FOnPackageSaved& UE::ConcertSyncCore::SyncDatabase::GetOnPackageSavedDelegate()
+{
+	static FOnPackageSaved SavedDelegate;
+	return SavedDelegate;
+}
+
 namespace TransactionDataUtil
 {
 
@@ -2740,24 +2745,35 @@ bool FConcertSyncSessionDatabase::LoadTransaction(const FString& InTransactionFi
 	return false;
 }
 
-bool HandleWritePackageResult(const FString& InDstPackageBlobPathname, TUniquePtr<FConcertFileCache>& Cache, PackageDataUtil::WritePackageResult Result)
+bool HandleWritePackageResult(const FName& PackageName, const FString& InDstPackageBlobPathname, TUniquePtr<FConcertFileCache>& Cache, PackageDataUtil::WritePackageResult Result)
 {
-	if (Result.CacheData.Num()>0)
+	auto SaveResult = [](const FString& InDstPackageBlobPathname, TUniquePtr<FConcertFileCache>& Cache, PackageDataUtil::WritePackageResult Result)
 	{
-		if (Result.bShouldCache)
+		if (Result.CacheData.Num()>0)
 		{
-			return Cache->SaveAndCacheFile(InDstPackageBlobPathname, MoveTemp(Result.CacheData));
-		}
-		TUniquePtr<FArchive> DstAr(IFileManager::Get().CreateFileWriter(*InDstPackageBlobPathname));
-		if (DstAr)
-		{
-			DstAr->Serialize(Result.CacheData.GetData(), Result.CacheData.Num());
-			return !DstAr->IsError();
+			if (Result.bShouldCache)
+			{
+				return Cache->SaveAndCacheFile(InDstPackageBlobPathname, MoveTemp(Result.CacheData));
+			}
+			TUniquePtr<FArchive> DstAr(IFileManager::Get().CreateFileWriter(*InDstPackageBlobPathname));
+			if (DstAr)
+			{
+				DstAr->Serialize(Result.CacheData.GetData(), Result.CacheData.Num());
+				return !DstAr->IsError();
+			}
+			UE_LOG(LogConcert, Warning, TEXT("Failed to handle package write result of file '%s'."), *InDstPackageBlobPathname);
+			return false;
 		}
-		UE_LOG(LogConcert, Warning, TEXT("Failed to handle package write result of file '%s'."), *InDstPackageBlobPathname);
-		return false;
+		return true;
+	};
+
+	const bool bSaveResult = SaveResult(InDstPackageBlobPathname, Cache, Result);
+	if (bSaveResult)
+	{
+		UE::ConcertSyncCore::SyncDatabase::GetOnPackageSavedDelegate().Broadcast(PackageName);
 	}
-	return true;
+
+	return bSaveResult;
 }
 
 bool FConcertSyncSessionDatabase::SavePackage(const FString& InDstPackageBlobPathname, const FConcertPackageInfo& InPackageInfo, FConcertPackageDataStream& InPackageDataStream)
@@ -2772,9 +2788,10 @@ bool FConcertSyncSessionDatabase::SavePackage(const FString& InDstPackageBlobPat
 	//     of this object and thus we cannot perform asynchronous work (Unless the raw data via DataBlob has already
 	//     been provided) in which case we ignore the DstAr and do asynchronous path.
 	//
+	FName PackageName = InPackageInfo.PackageUpdateType == EConcertPackageUpdateType::Renamed ? InPackageInfo.NewPackageName : InPackageInfo.PackageName;
 	if (InPackageDataStream.DataBlob && InPackageDataStream.DataBlob->Num() > 0)
 	{
-		ScheduleAsyncWrite(InDstPackageBlobPathname, InPackageDataStream);
+		ScheduleAsyncWrite(PackageName, InDstPackageBlobPathname, InPackageDataStream);
 		return true;
 	}
 	else
@@ -2783,7 +2800,7 @@ bool FConcertSyncSessionDatabase::SavePackage(const FString& InDstPackageBlobPat
 		{
 			FArchive NullArchive;
 			PackageDataUtil::WritePackageResult Result = PackageDataUtil::WritePackage(InPackageDataStream, &NullArchive);
-			return HandleWritePackageResult(InDstPackageBlobPathname, PackageFileCache, MoveTemp(Result));
+			return HandleWritePackageResult(PackageName, *InDstPackageBlobPathname, PackageFileCache, MoveTemp(Result));
 		}
 		else
 		{
@@ -2835,11 +2852,11 @@ void FConcertSyncSessionDatabase::FlushAsynchronousTasks()
 	{
 		if (Item.Value->AsyncTask.Get())
 		{
-			HandleWritePackageResult(Item.Key, PackageFileCache, MoveTemp(Item.Value->Result));
+			HandleWritePackageResult(Item.Value->CachedPackageName, Item.Key, PackageFileCache, MoveTemp(Item.Value->Result));
 		}
 		else
 		{
-			UE_LOG(LogConcert, Error, TEXT("Async task failed to write package %s"), *Item.Value->CachedPackageName);
+			UE_LOG(LogConcert, Error, TEXT("Async task failed to write package %s"), *Item.Value->CachedPackageName.ToString());
 		}
 	}
 	DeferredLargePackageIOPtr->Clear();
@@ -2851,7 +2868,7 @@ void FConcertSyncSessionDatabase::UpdateAsynchronousTasks()
 	{
 		if (It->Value->AsyncTask.IsReady())
 		{
-			HandleWritePackageResult(It->Key, PackageFileCache,MoveTemp(It->Value->Result));
+			HandleWritePackageResult(It->Value->CachedPackageName, It->Key, PackageFileCache, MoveTemp(It->Value->Result));
 			It.RemoveCurrent();
 		}
 	}
@@ -2893,15 +2910,15 @@ TOptional<int64> FConcertSyncSessionDatabase::GetSpecifiedOrHeadPackageRevision(
 	return PackageRevision;
 }
 
-void FConcertSyncSessionDatabase::ScheduleAsyncWrite(const FString& InDstPackageBlobPathname, FConcertPackageDataStream& InPackageDataStream)
+void FConcertSyncSessionDatabase::ScheduleAsyncWrite(const FName& InPackageName, const FString& InDstPackageBlobPathname, FConcertPackageDataStream& InPackageDataStream)
 {
 	check(InPackageDataStream.DataBlob && InPackageDataStream.DataBlob->Num() > 0);
 
-	FConcertPackageAsyncDataStreamPtr SharedStream = MakeShared<FConcertPackageAsyncDataStream,ESPMode::ThreadSafe>(InPackageDataStream, InDstPackageBlobPathname);
+	FConcertPackageAsyncDataStreamPtr SharedStream = MakeShared<FConcertPackageAsyncDataStream,ESPMode::ThreadSafe>(InPackageDataStream, InPackageName);
 	if(FConcertPackageAsyncDataStreamPtr* Item = DeferredLargePackageIOPtr->GetMap().Find(InDstPackageBlobPathname))
 	{
 		(*Item)->AsyncTask.Get();
-		HandleWritePackageResult(InDstPackageBlobPathname, PackageFileCache, MoveTemp((*Item)->Result));
+		HandleWritePackageResult((*Item)->CachedPackageName, InDstPackageBlobPathname, PackageFileCache, MoveTemp((*Item)->Result));
 	}
 	DeferredLargePackageIOPtr->GetMap().FindOrAdd(InDstPackageBlobPathname) = SharedStream;
 	SharedStream->AsyncTask = Async(EAsyncExecution::TaskGraph, [InDstPackageBlobPathname,SharedStream]()
diff --git a/Engine/Plugins/Developer/Concert/ConcertSync/ConcertSyncCore/Source/ConcertSyncCore/Public/ConcertSyncSessionDatabase.h b/Engine/Plugins/Developer/Concert/ConcertSync/ConcertSyncCore/Source/ConcertSyncCore/Public/ConcertSyncSessionDatabase.h
--- a/Engine/Plugins/Developer/Concert/ConcertSync/ConcertSyncCore/Source/ConcertSyncCore/Public/ConcertSyncSessionDatabase.h
+++ b/Engine/Plugins/Developer/Concert/ConcertSync/ConcertSyncCore/Source/ConcertSyncCore/Public/ConcertSyncSessionDatabase.h
@@ -23,6 +23,7 @@ using FIteratePackageActivityFunc = TFunctionRef<EBreakBehavior(FConcertSyncActi
 using FIterateActivityFunc = TFunctionRef<EBreakBehavior(FConcertSyncActivity&&)>;
 
 DECLARE_MULTICAST_DELEGATE_OneParam(FOnActivityProduced, const FConcertSyncActivity&);
+DECLARE_MULTICAST_DELEGATE_OneParam(FOnPackageSaved, const FName& PackageName);
 
 /**
  * Database of activities that have happened in a Concert Sync Session.
@@ -695,15 +696,16 @@ private:
 
 	/** Helper functions that for getting a package revision for an optional package revision argument */
 	TOptional<int64> GetSpecifiedOrHeadPackageRevision(FName InPackageName, const int64* InPackageRevision) const;
-	
+
 	/**
 	 * Schedule an asynchronous write for the given Package Stream.  The stream must be in-memory. File sharing
 	 * asynchronous write is not supported.
 	 *
+	 * @param InPackageName	            Name of the package in content area.
 	 * @param InDstPackageBlobPathName  Full path of the destination package.
 	 * @param InPackageDataStream       The package data stream.
 	 **/
-	void ScheduleAsyncWrite(const FString& InDstPackageBlobPathname, FConcertPackageDataStream& InPackageDataStream);
+	void ScheduleAsyncWrite(const FName& InPackageName, const FString& InDstPackageBlobPathname, FConcertPackageDataStream& InPackageDataStream);
 
 	/**
 	 * Set the active ignored state for the given activity.
@@ -970,6 +972,12 @@ private:
 	TPimplPtr<struct FDeferredLargePackageIOImpl> DeferredLargePackageIOPtr;
 };
 
+namespace UE::ConcertSyncCore::SyncDatabase
+{
+	/** Get the global saved package delegate */
+	CONCERTSYNCCORE_API FOnPackageSaved& GetOnPackageSavedDelegate();
+}
+
 namespace ConcertSyncSessionDatabaseFilterUtil
 {
 
