Add support for different compression types on database and event handling.

diff --git a/Engine/Plugins/Developer/Concert/ConcertMain/Source/Concert/Private/ConcertMessageData.cpp b/Engine/Plugins/Developer/Concert/ConcertMain/Source/Concert/Private/ConcertMessageData.cpp
--- a/Engine/Plugins/Developer/Concert/ConcertMain/Source/Concert/Private/ConcertMessageData.cpp
+++ b/Engine/Plugins/Developer/Concert/ConcertMain/Source/Concert/Private/ConcertMessageData.cpp
@@ -2,6 +2,7 @@
 
 #include "ConcertMessageData.h"
 #include "ConcertLogGlobal.h"
+#include "HAL/IConsoleManager.h"
 #include "IdentifierTable/ConcertTransportArchives.h"
 
 #include "Misc/App.h"
@@ -167,23 +168,50 @@ bool FConcertSessionFilter::ActivityIdPassesFilter(const int64 InActivityId) con
 		&& ActivityIdUpperBound >= InActivityId;
 }
 
+namespace UE::Concert::Compression
+{
+
+static TAutoConsoleVariable<int32> CVarCompressionType(
+	TEXT("Concert.SetCompressionType"), 1,
+	TEXT("Specify the type of compression to use when serializing data. A value of 0 means compression is off. A value of 1 = Oodle. All other values = Zlib."));
+
+static TAutoConsoleVariable<int32> CVarCompressionSizeLimit(
+	TEXT("Concert.SetCompressionSizeLimit"), 32,
+	TEXT("The compressor incurs a performance cost and will only be used if the package size is less than the amount specified (default 32 MB). A value of 0 or less will always compress."));
+
+static TAutoConsoleVariable<int32> CVarCompressionFlags(
+	TEXT("Concert.SetCompressionFlags"), 0,
+	TEXT("Specify the flags to use when compression is enabled. A value of 0 means no flags. A value of 1 favors smaller sizes. Any other value favors faster encoding."));
+
+int32 GetConsoleVariableCompressionType()
+{
+	return CVarCompressionType.GetValueOnAnyThread();
+}
+
+int32 GetConsoleVariableCompressionFlags()
+{
+	return CVarCompressionFlags.GetValueOnAnyThread();
+}
+
+int32 GetConsoleVariableCompressionSizeLimit()
+{
+	return CVarCompressionSizeLimit.GetValueOnAnyThread();
+}
+
+}
+
 namespace PayloadDetail
 {
 
 bool ShouldCompress(const FConcertSessionSerializedPayload& InPayload, EConcertPayloadCompressionType CompressionType)
 {
-	// We need to have at least 512 bytes to invoke the compressor.
-	constexpr int32 MinBytes = 512;
-	// Set the heuristic limit to 3 MB
-	constexpr int32 MaxBytes = 3 * 1024 * 1024;
-
 	if (CompressionType == EConcertPayloadCompressionType::None)
 	{
 		return false;
 	}
 	if (CompressionType == EConcertPayloadCompressionType::Heuristic)
 	{
-		return InPayload.PayloadSize > MinBytes && InPayload.PayloadSize < MaxBytes;
+		return UE::Concert::Compression::ShouldCompress(InPayload.PayloadSize);
 	}
 	check(CompressionType == EConcertPayloadCompressionType::Always);
 	// Otherwise we are always compressing
@@ -200,24 +228,27 @@ bool TryCompressImpl(const UScriptStruct* InEventType, const void* InEventData,
 		TArray<uint8> &InBytes = InOutPayload.PayloadBytes.Bytes;
 		TArray<uint8> OutCompressedData;
 		// Compress the result to send on the wire
-		int32 CompressedSize = FCompression::CompressMemoryBound(NAME_Zlib, InOutPayload.PayloadSize);
+		FName NamedCompressionAlgo = UE::Concert::Compression::GetCompressionAlgorithm();
+		ECompressionFlags CompressFlags = UE::Concert::Compression::GetCompressionFlags();
+		int32 CompressedSize = FCompression::CompressMemoryBound(NamedCompressionAlgo, InOutPayload.PayloadSize, CompressFlags);
 		OutCompressedData.SetNumUninitialized(CompressedSize);
 		SCOPED_CONCERT_TRACE(SerializePayload_CompressMemory);
-		if (FCompression::CompressMemory(NAME_Zlib, OutCompressedData.GetData(), CompressedSize, InBytes.GetData(), InBytes.Num()))
+		if (FCompression::CompressMemory(NamedCompressionAlgo, OutCompressedData.GetData(), CompressedSize,
+										 InBytes.GetData(), InBytes.Num()), CompressFlags)
 		{
 			OutCompressedData.SetNum(CompressedSize, false);
 			InOutPayload.PayloadBytes.Bytes = MoveTemp(OutCompressedData);
-			InOutPayload.bPayloadIsCompressed = true;
+			InOutPayload.PayloadCompressionDetails = UE::Concert::Compression::GetCompressionFromNamedType(NamedCompressionAlgo, CompressFlags);
 		}
 		else
 		{
 			UE_LOG(LogConcert, Warning, TEXT("Unable to compress data for %s!"), *InEventType->GetName());
-			InOutPayload.bPayloadIsCompressed = false;
+			InOutPayload.PayloadCompressionDetails = EConcertCompressionDetails::Uncompressed;
 		}
 	}
 	else
 	{
-		InOutPayload.bPayloadIsCompressed = false;
+		InOutPayload.PayloadCompressionDetails = EConcertCompressionDetails::Uncompressed;
 	}
 
 	// Since we can support uncompressed or compressed data this is always successful.
@@ -228,18 +259,24 @@ using OptionalDecompressBytes = TOptional<TArray<uint8>>;
 
 OptionalDecompressBytes DecompressImpl(const FConcertSessionSerializedPayload& InPayload)
 {
-	if (InPayload.bPayloadIsCompressed)
+	EConcertCompressionDetails CompressAlgo = InPayload.PayloadCompressionDetails;
+	if (UE::Concert::Compression::DataIsCompressed(CompressAlgo))
 	{
 		const TArray<uint8> &InBytes = InPayload.PayloadBytes.Bytes;
 		TArray<uint8> UncompressedData;
 		UncompressedData.SetNumUninitialized(InPayload.PayloadSize);
 		SCOPED_CONCERT_TRACE(DeserializePayload_UncompressMemory);
-		if (FCompression::UncompressMemory(NAME_Zlib, UncompressedData.GetData(), UncompressedData.Num(), InBytes.GetData(), InBytes.Num()))
+		ECompressionFlags CompressFlags = UE::Concert::Compression::GetCoreCompressionFlags(CompressAlgo);
+		FName CompressType = UE::Concert::Compression::GetCompressionAlgorithm(CompressAlgo);
+		if (FCompression::UncompressMemory(CompressType, UncompressedData.GetData(), UncompressedData.Num(), InBytes.GetData(), InBytes.Num(), CompressFlags))
 		{
 			return OptionalDecompressBytes(MoveTemp(UncompressedData));
 		}
 	}
 
+	check(UE::Concert::Compression::DataIsUncompressed(CompressAlgo) &&
+		  InPayload.PayloadBytes.Bytes.Num() == InPayload.PayloadSize);
+
 	return OptionalDecompressBytes{};
 }
 
@@ -281,7 +318,9 @@ bool DeserializeImpl(const UScriptStruct* InTargetEventType, void* InOutTargetEv
 bool DeserializeAndDecompress(const UScriptStruct* InTargetEventType, void* InOutTargetEventData, const FConcertSessionSerializedPayload& InPayload)
 {
 	OptionalDecompressBytes DecompressedBytes = DecompressImpl(InPayload);
-	if ( InPayload.bPayloadIsCompressed && !DecompressedBytes.IsSet() )
+	EConcertCompressionDetails CompressAlgo = InPayload.PayloadCompressionDetails;
+
+	if ( UE::Concert::Compression::DataIsCompressed(CompressAlgo) && !DecompressedBytes.IsSet() )
 	{
 		return false;
 	}
diff --git a/Engine/Plugins/Developer/Concert/ConcertMain/Source/Concert/Public/ConcertMessageData.h b/Engine/Plugins/Developer/Concert/ConcertMain/Source/Concert/Public/ConcertMessageData.h
--- a/Engine/Plugins/Developer/Concert/ConcertMain/Source/Concert/Public/ConcertMessageData.h
+++ b/Engine/Plugins/Developer/Concert/ConcertMain/Source/Concert/Public/ConcertMessageData.h
@@ -3,6 +3,7 @@
 #pragma once
 
 #include "CoreMinimal.h"
+#include "Misc/CompressionFlags.h"
 #include "UObject/ObjectMacros.h"
 #include "Misc/Guid.h"
 #include "ConcertVersion.h"
@@ -44,6 +45,181 @@ enum class EConcertPayloadCompressionType : uint8
 	Always
 };
 
+UENUM()
+enum class EConcertCompressionDetails : uint8
+{
+	/** The package data (the body) is written uncompressed. */
+	Uncompressed = 0,
+
+	/** The package data (the body) is written in a single compressed block. The value 1 represents the original format. Still used to store small packages.*/
+	Compressed = 1 << 0,
+
+	/** If compression is enabled and this bit is set then oodle compression is used for the package data. */
+	CompressWithOodle = 1 << 1,
+
+	/** Compress for speed. */
+	CompressForSpeed = 1 << 2,
+
+	/** Compress for smaller sizes. */
+	CompressForSize = 1 << 3
+};
+ENUM_CLASS_FLAGS(EConcertCompressionDetails)
+
+namespace UE::Concert::Compression
+{
+
+/** Get the compression type from the console variable setting. */
+CONCERT_API int32 GetConsoleVariableCompressionType();
+
+/** Function to check the console variable for the compression flags. Speed vs size. */
+CONCERT_API int32 GetConsoleVariableCompressionFlags();
+
+/** Function to check the console variable for compression size limit.*/
+CONCERT_API int32 GetConsoleVariableCompressionSizeLimit();
+
+/** Based on the size of the data to compress indicate if we should invoke the compressor. */
+template <typename SizeType>
+inline bool ShouldCompress(SizeType DataSize)
+{
+	static_assert(std::is_integral<SizeType>::value);
+	const int32 CompressionLimit = GetConsoleVariableCompressionSizeLimit();
+	const int32 MinLimit = 512;
+	if (CompressionLimit <= 0)
+	{
+		return DataSize > MinLimit;
+	}
+	const SizeType MaxPackageDataSizeForCompression = static_cast<SizeType>(CompressionLimit) * 1024 * 1024;
+	return DataSize > MinLimit && DataSize <= MaxPackageDataSizeForCompression;
+}
+
+/** Returns truen if the compression details is set to compressed format.  */
+inline bool DataIsCompressed(EConcertCompressionDetails InFormat)
+{
+	return EnumHasAllFlags(InFormat, EConcertCompressionDetails::Compressed);
+}
+
+/** Returns true if the compression details is set to uncompressed format. */
+inline bool DataIsUncompressed(EConcertCompressionDetails InFormat)
+{
+	return InFormat == EConcertCompressionDetails::Uncompressed;
+}
+
+/** Returns true if we are suppose to use Oodle for compression. */
+inline bool DataIsCompressedWithOodle(EConcertCompressionDetails InFormat)
+{
+	return EnumHasAllFlags(InFormat, EConcertCompressionDetails::Compressed | EConcertCompressionDetails::CompressWithOodle);
+}
+
+/** Get the named compression algorithm to invoke with serializer and memory compressors. */
+inline FName GetCompressionAlgorithm()
+{
+	const int32 CompressionAlgo = GetConsoleVariableCompressionType();
+	if (CompressionAlgo == 1)
+	{
+		return NAME_Oodle;
+	}
+	return NAME_Zlib;
+}
+
+/** Get the default flags to use when invoking the compressor. */
+inline ECompressionFlags GetCompressionFlags()
+{
+	const int32 CompressionFlags = GetConsoleVariableCompressionFlags();
+	if (CompressionFlags == 0)
+	{
+		return COMPRESS_NoFlags;
+	}
+	if (CompressionFlags == 1)
+	{
+		return COMPRESS_BiasSize;
+	}
+	return COMPRESS_BiasSpeed;
+}
+
+/** Get the named compression algorithm based on the provided details.  Current this is only Zlib or Oodle. */
+inline FName GetCompressionAlgorithm(EConcertCompressionDetails InDetails)
+{
+	return DataIsCompressedWithOodle(InDetails) ? NAME_Oodle : NAME_Zlib;
+}
+
+/** Get the compression enum value given the concert compression settings.  This the value passed into the serializer and memory compressor. */
+inline ECompressionFlags GetCoreCompressionFlags(EConcertCompressionDetails InFormat)
+{
+	// Make sure we do not have both flags set.
+	check(!EnumHasAllFlags(InFormat, EConcertCompressionDetails::CompressForSize | EConcertCompressionDetails::CompressForSpeed));
+	if (DataIsUncompressed(InFormat))
+	{
+		return COMPRESS_NoFlags;
+	}
+	if (EnumHasAnyFlags(InFormat, EConcertCompressionDetails::CompressForSize))
+	{
+		return COMPRESS_BiasSize;
+	}
+	if (EnumHasAnyFlags(InFormat, EConcertCompressionDetails::CompressForSpeed))
+	{
+		return COMPRESS_BiasSpeed;
+	}
+	return COMPRESS_NoFlags;
+}
+
+/** For a given named compression algorithm and compression flags convert it into a EConcertCompressionDetails */
+inline EConcertCompressionDetails GetCompressionFromNamedType(FName NamedMethod, ECompressionFlags Flags)
+{
+	check(NamedMethod == NAME_Oodle || NamedMethod == NAME_Zlib);
+	check(Flags == COMPRESS_NoFlags || Flags == COMPRESS_BiasMemory || Flags == COMPRESS_BiasSpeed);
+
+	EConcertCompressionDetails Compressor = EConcertCompressionDetails::Compressed;
+
+	if (NamedMethod == NAME_Oodle)
+	{
+		Compressor = Compressor | EConcertCompressionDetails::CompressWithOodle;
+	}
+
+	if (Flags == COMPRESS_BiasMemory)
+	{
+		Compressor = Compressor | EConcertCompressionDetails::CompressForSize;
+	}
+	else if (Flags == COMPRESS_BiasSpeed)
+	{
+		Compressor = Compressor | EConcertCompressionDetails::CompressForSpeed;
+	}
+	return Compressor;
+}
+
+/** Add in the compression flags to the concert compression details enum. */
+inline EConcertCompressionDetails GetCompressionFlags(EConcertCompressionDetails InFormat)
+{
+	const int32 CompressionFlags = GetConsoleVariableCompressionFlags();
+	if (CompressionFlags == 0)
+	{
+		return InFormat;
+	}
+	if (CompressionFlags == 1)
+	{
+		return InFormat | EConcertCompressionDetails::CompressForSize;
+	}
+	return InFormat | EConcertCompressionDetails::CompressForSpeed;
+}
+
+/** Based on the data size return the current compression details to store. */
+template <typename SizeType>
+inline EConcertCompressionDetails GetCompressionDetails(SizeType DataSize)
+{
+	static_assert(std::is_integral<SizeType>::value);
+	const int32 CompressionType = GetConsoleVariableCompressionType();
+	if (ShouldCompress(DataSize) && CompressionType != 0)
+	{
+		if (CompressionType == 1)
+		{
+			return GetCompressionFlags(EConcertCompressionDetails::Compressed | EConcertCompressionDetails::CompressWithOodle);
+		}
+		return GetCompressionFlags(EConcertCompressionDetails::Compressed);
+	}
+	return EConcertCompressionDetails::Uncompressed;
+}
+
+}
+
 UENUM()
 enum class EConcertPayloadSerializationMethod : uint8
 {
@@ -342,7 +518,7 @@ struct FConcertSessionSerializedPayload
 	CONCERT_API bool GetPayload(const UScriptStruct* InPayloadType, void* InOutPayloadData) const;
 
 	CONCERT_API bool IsTypeChildOf(const UScriptStruct* InPayloadType) const;
-	
+
 	template<typename T>
 	bool IsTypeChildOf() const
 	{
@@ -355,6 +531,12 @@ struct FConcertSessionSerializedPayload
 		return GetPayload(TBaseStructure<T>::Get(), &OutPayloadData);
 	}
 
+	/** Indicates if the stored payload has been compressed with any compression algorithm. */
+	bool PayloadIsCompressed() const
+	{
+		return UE::Concert::Compression::DataIsCompressed(PayloadCompressionDetails);
+	}
+
 	/** The typename of the user-defined payload. */
 	UPROPERTY(VisibleAnywhere, Category="Payload")
 	FName PayloadTypeName;
@@ -363,9 +545,9 @@ struct FConcertSessionSerializedPayload
 	UPROPERTY(VisibleAnywhere, Category = "Payload")
 	EConcertPayloadSerializationMethod SerializationMethod = EConcertPayloadSerializationMethod::Standard;
 
-	/** Indicates if the serialized payload has been compressed. */
-	UPROPERTY(VisibleAnywhere, Category="Payload")
-	bool bPayloadIsCompressed = false;
+	/** Indicates how the serialized payload has been compressed. */
+	UPROPERTY(VisibleAnywhere, Category = "Payload")
+	EConcertCompressionDetails PayloadCompressionDetails = EConcertCompressionDetails::Uncompressed;
 
 	/** The uncompressed size of the user-defined payload data. */
 	UPROPERTY(VisibleAnywhere, Category="Payload")
diff --git a/Engine/Plugins/Developer/Concert/ConcertSync/ConcertSyncClient/Source/ConcertSyncClient/Private/ConcertClientDataStore.cpp b/Engine/Plugins/Developer/Concert/ConcertSync/ConcertSyncClient/Source/ConcertSyncClient/Private/ConcertClientDataStore.cpp
--- a/Engine/Plugins/Developer/Concert/ConcertSync/ConcertSyncClient/Source/ConcertSyncClient/Private/ConcertClientDataStore.cpp
+++ b/Engine/Plugins/Developer/Concert/ConcertSync/ConcertSyncClient/Source/ConcertSyncClient/Private/ConcertClientDataStore.cpp
@@ -1,6 +1,7 @@
 // Copyright Epic Games, Inc. All Rights Reserved.
 
 #include "ConcertClientDataStore.h"
+#include "ConcertMessageData.h"
 #include "IConcertSession.h"
 #include "ConcertSyncClientLiveSession.h"
 
@@ -98,7 +99,7 @@ TFuture<FConcertDataStoreResult> FConcertClientDataStore::InternalCompareExchang
 			// Send the 'version' rather than 'expected' to save bandwidth.
 			CompareExchangeRequest.ExpectedVersion = Result.Value->Version;
 			CompareExchangeRequest.Expected.PayloadTypeName = TEXT("");
-			CompareExchangeRequest.Expected.bPayloadIsCompressed = false;
+			CompareExchangeRequest.Expected.PayloadCompressionDetails = EConcertCompressionDetails::Uncompressed;
 			CompareExchangeRequest.Expected.PayloadSize = 0;
 			CompareExchangeRequest.Expected.PayloadBytes.Bytes.Reset();
 		}
@@ -184,7 +185,7 @@ FConcertDataStoreResult FConcertClientDataStore::HandleResponse(const FName& Sen
 		// Ensure the server doesn't send data when the client initiated the operation (The client knows the value it sent).
 		check(Response.Value.TypeName.IsNone());
 		check(Response.Value.SerializedValue.PayloadTypeName.IsNone());
-		check(Response.Value.SerializedValue.bPayloadIsCompressed == false);
+		check(Response.Value.SerializedValue.PayloadCompressionDetails == EConcertCompressionDetails::Uncompressed);
 		check(Response.Value.SerializedValue.PayloadBytes.Bytes.Num() == 0);
 
 		// Ensure the server sent back a valid version in case it has exchanged the value. (Successfully added value is always version 1)
diff --git a/Engine/Plugins/Developer/Concert/ConcertSync/ConcertSyncCore/Source/ConcertSyncCore/Private/ConcertSyncSessionDatabase.cpp b/Engine/Plugins/Developer/Concert/ConcertSync/ConcertSyncCore/Source/ConcertSyncCore/Private/ConcertSyncSessionDatabase.cpp
--- a/Engine/Plugins/Developer/Concert/ConcertSync/ConcertSyncCore/Source/ConcertSyncCore/Private/ConcertSyncSessionDatabase.cpp
+++ b/Engine/Plugins/Developer/Concert/ConcertSync/ConcertSyncCore/Source/ConcertSyncCore/Private/ConcertSyncSessionDatabase.cpp
@@ -7,6 +7,7 @@
 #include "ConcertSyncSessionTypes.h"
 #include "ConcertUtil.h"
 
+#include "Misc/CompressionFlags.h"
 #include "SQLiteDatabase.h"
 #include "HAL/FileManager.h"
 #include "Templates/SharedPointer.h"
@@ -81,7 +82,7 @@ namespace TransactionDataUtil
 const int32 MinFilesToCache = 10;
 const uint64 MaxFileSizeBytesToCache = 50 * 1024 * 1024;
 const int64 BucketSize = 500;
-const uint32 DataVersion = 1;
+const uint32 DataVersion = 2;
 const FGuid DataEntryFooter = FGuid(0xE473C070, 0x65DA42BF, 0xA0607C78, 0xE0DC47CF);
 
 FString GetDataPath(const FString& InSessionPath)
@@ -131,7 +132,19 @@ bool WriteTransaction(const FStructOnScope& InTransaction, TArray<uint8>& OutSer
 	Ar.SerializeIntPacked(UncompressedTransactionSize);
 	if (UncompressedTransactionSize > 0)
 	{
-		Ar.SerializeCompressed(UncompressedTransaction.GetData(), UncompressedTransactionSize, NAME_Zlib);
+		EConcertCompressionDetails StoredFormat = UE::Concert::Compression::GetCompressionDetails(UncompressedTransactionSize);
+		uint32 StoredFormatAsUint = static_cast<uint32>(StoredFormat);
+		Ar.SerializeIntPacked(StoredFormatAsUint);
+		if (UE::Concert::Compression::DataIsCompressed(StoredFormat))
+		{
+			Ar.SerializeCompressed(UncompressedTransaction.GetData(), UncompressedTransactionSize,
+								   UE::Concert::Compression::GetCompressionAlgorithm(StoredFormat),
+								   UE::Concert::Compression::GetCoreCompressionFlags(StoredFormat));
+		}
+		else
+		{
+			Ar.Serialize(UncompressedTransaction.GetData(), UncompressedTransactionSize);
+		}
 	}
 
 	// Serialize the footer so we know we didn't crash mid-write
@@ -206,9 +219,28 @@ bool ReadTransaction(const TArray<uint8>& InSerializedTransactionData, FStructOn
 	Ar.SerializeIntPacked(UncompressedTransactionSize);
 	TArray<uint8> UncompressedTransaction;
 	UncompressedTransaction.AddZeroed(UncompressedTransactionSize);
+
 	if (UncompressedTransactionSize > 0)
 	{
-		Ar.SerializeCompressed(UncompressedTransaction.GetData(), UncompressedTransactionSize, NAME_Zlib);
+		// Version 1 data is stored always compressed in Zlib format.
+		EConcertCompressionDetails Format = EConcertCompressionDetails::Compressed;
+		if (SerializedDataVersion > 1)
+		{
+			uint32 SerializedDataFormat = 0;
+			Ar.SerializeIntPacked(SerializedDataFormat);
+			uint8 DataFormatAsUint8 = static_cast<uint8>(SerializedDataFormat & 0xFF);
+			Format = static_cast<EConcertCompressionDetails>(DataFormatAsUint8);
+		}
+		if (UE::Concert::Compression::DataIsCompressed(Format))
+		{
+			Ar.SerializeCompressed(UncompressedTransaction.GetData(), UncompressedTransactionSize,
+								   UE::Concert::Compression::GetCompressionAlgorithm(Format),
+								   UE::Concert::Compression::GetCoreCompressionFlags(Format));
+		}
+		else
+		{
+			Ar.Serialize(UncompressedTransaction.GetData(), UncompressedTransactionSize);
+		}
 	}
 
 	// Read the raw transaction data
@@ -234,10 +266,6 @@ const FGuid EntryFooter = FGuid(0x2EFC8CDD, 0x748E46C0, 0xA5485769, 0x13A3C354);
 constexpr const uint64 MaxPackageBlobSizeForCaching = FMath::Max(
 	PackageDataUtil::MaxFileSizeBytesToCache / PackageDataUtil::MinFilesToCache, 16ull * 1024 * 1024);
 
-// Bigger files will not be compressed to avoid performance hit to decompress them when requested.
-// (Streaming compression is not supported yet, so it is done in memory only)
-//
-constexpr const uint64 MaxPackageDataSizeForCompression = 32ull * 1024 * 1024;
 
 /**
  * Returns true if a package blob should be cached in memory, according to its size.
@@ -249,41 +277,8 @@ bool ShouldCachePackageBlob(uint64 PackageBlobSize)
 	return PackageBlobSize <= MaxPackageBlobSizeForCaching;
 }
 
-/**
- * Returns true if the package data should be compressed in the blob, according to its size.
- * @param PackageDataSize The uncompressed package data size.
- * @return true if the package data should be compressed, false otherwise.
- */
-bool ShouldCompressPackageData(uint64 PackageDataSize)
-{
-	return PackageDataSize > 0 && PackageDataSize <= MaxPackageDataSizeForCompression;
-}
-
-/**
- * Defines how the package data file written. The 8 LSB are reserved for FileFormatVersion which control how 24 MSB are interpreted.
- */
-enum class EPackageDataFormat : uint8
-{
-	/** The package data (the body) is written in a single compressed block. The value 1 represents the original format. Still used to store small packages.*/
-	PackageDataCompressed = 1 << 0,
-
-	/** The package data (the body) is written uncompressed. Added at 4.25 to support large package files. */
-	PackageDataUncompressed = 1 << 1,
-};
-ENUM_CLASS_FLAGS(EPackageDataFormat)
-
-
-PackageDataUtil::EPackageDataFormat GetPackageDataFormat(int64 DataSize)
-{
-	if (ShouldCompressPackageData(DataSize))
-	{
-		return PackageDataUtil::EPackageDataFormat::PackageDataCompressed;
-	}
-	return PackageDataUtil::EPackageDataFormat::PackageDataUncompressed;
-}
-
 /** Write the version/format information in the first 4 bytes of the blob, enabling the reader to parse the blob. */
-void WritePackageBlobVersionInfo(FArchive& Ar, EPackageDataFormat PackageDataFormat)
+void WritePackageBlobVersionInfo(FArchive& Ar, EConcertCompressionDetails PackageDataFormat)
 {
 	check(Ar.IsSaving() && Ar.Tell() == 0); // Expects to write the version/format info at the beginning.
 
@@ -294,7 +289,7 @@ void WritePackageBlobVersionInfo(FArchive& Ar, EPackageDataFormat PackageDataFor
 }
 
 /** Reads the first 4 bytes of the archive and detect how the package data was written in the blob. */
-EPackageDataFormat ParsePackageDataFormat(FArchive& Ar)
+EConcertCompressionDetails ParsePackageDataFormat(FArchive& Ar)
 {
 	check(Ar.IsLoading() && Ar.Tell() == 0); // Expected to read the version/format info at the beginning.
 
@@ -304,12 +299,12 @@ EPackageDataFormat ParsePackageDataFormat(FArchive& Ar)
 	if (ExtractedDataVersion == 1)
 	{
 		// In Version 1, the package was always written compressed into a single block.
-		return EPackageDataFormat::PackageDataCompressed;
+		return EConcertCompressionDetails::Compressed;
 	}
 
 	// In Version 2, the package data format was encoded after the 8 LSB used for 'DataVersion'
 	constexpr uint32 BitsPerByte = 8;
-	return static_cast<EPackageDataFormat>(SerializedDataVersion >> (sizeof(DataVersion) * BitsPerByte));
+	return static_cast<EConcertCompressionDetails>(SerializedDataVersion >> (sizeof(DataVersion) * BitsPerByte));
 }
 
 FString GetDataPath(const FString& InSessionPath)
@@ -328,7 +323,7 @@ FString GetDataFilename(const FName InPackageName, const int64 InRevision)
 }
 
 /** Internally used by WritePackageBlob() - write the serialized package header. */
-void WritePackageHeader(FArchive& Ar, EPackageDataFormat PackageDataFormat)
+void WritePackageHeader(FArchive& Ar, EConcertCompressionDetails PackageDataFormat)
 {
 	// First 4-bytes contains information about the format.
 	WritePackageBlobVersionInfo(Ar, PackageDataFormat);
@@ -343,11 +338,11 @@ void WritePackageHeader(FArchive& Ar, EPackageDataFormat PackageDataFormat)
 }
 
 /** Internally used by ExtractPackageData() - read serialized package header. */
-TTuple<EPackageDataFormat, int64> ReadPackageHeader(FArchive& Ar)
+TTuple<EConcertCompressionDetails, int64> ReadPackageHeader(FArchive& Ar)
 {
 	// Consume the first bytes to read the package version and package data format.
-	EPackageDataFormat PackageDataFormat = ParsePackageDataFormat(Ar);
-	
+	EConcertCompressionDetails PackageDataFormat = ParsePackageDataFormat(Ar);
+
 	// Deserialize the next 8 bytes containing the offset to the package data (for binary compatibility with 4.24)
 	int64 PackageDataOffset = 0;
 	Ar << PackageDataOffset;
@@ -385,11 +380,12 @@ int64 EstimatePackageBlobMetaDataSize()
 
 bool WillUseMemoryWriter(int64 DataSize)
 {
-	return DataSize > 0 && (GetPackageDataFormat(DataSize) == EPackageDataFormat::PackageDataCompressed ||
+	using namespace UE::Concert::Compression;
+	return DataSize > 0 && (DataIsCompressed(GetCompressionDetails(DataSize)) ||
 							ShouldCachePackageBlob(DataSize + PackageDataUtil::EstimatePackageBlobMetaDataSize()));
 }
 
-void WritePackageCompressed(FConcertPackageDataStream& PackageDataStream, FArchive& DstAr)
+void WritePackageCompressed(FConcertPackageDataStream& PackageDataStream, FArchive& DstAr, EConcertCompressionDetails DataFormat)
 {
 	check(PackageDataStream.DataSize > 0);
 
@@ -400,14 +396,18 @@ void WritePackageCompressed(FConcertPackageDataStream& PackageDataStream, FArchi
 
 	if (PackageDataStream.DataBlob) // Optimization to avoid transfering from the archive to a temporary buffer.
 	{
-		DstAr.SerializeCompressed(const_cast<uint8*>(PackageDataStream.DataBlob->GetData()), UncompressedPackageSize, NAME_Zlib); // Write
+		DstAr.SerializeCompressed(const_cast<uint8*>(PackageDataStream.DataBlob->GetData()), UncompressedPackageSize,
+								  UE::Concert::Compression::GetCompressionAlgorithm(DataFormat),
+								  UE::Concert::Compression::GetCoreCompressionFlags(DataFormat));
 	}
 	else
 	{
 		TArray<uint8> UncompressedPackageData;
 		UncompressedPackageData.AddUninitialized(UncompressedPackageSize);
 		PackageDataStream.DataAr->Serialize(UncompressedPackageData.GetData(), UncompressedPackageSize); // Read.
-		DstAr.SerializeCompressed(const_cast<uint8*>(UncompressedPackageData.GetData()), UncompressedPackageSize, NAME_Zlib); // Write.
+		DstAr.SerializeCompressed(const_cast<uint8*>(UncompressedPackageData.GetData()), UncompressedPackageSize,
+								  UE::Concert::Compression::GetCompressionAlgorithm(DataFormat),
+								  UE::Concert::Compression::GetCoreCompressionFlags(DataFormat));
 	}
 }
 
@@ -433,7 +433,7 @@ WritePackageResult WritePackage(FConcertPackageDataStream& PackageDataStream, FA
 	WritePackageResult OutResult;
 	FMemoryWriter MemWriter(OutResult.CacheData);
 
-	EPackageDataFormat PackageDataFormat = GetPackageDataFormat(PackageDataStream.DataSize);
+	EConcertCompressionDetails PackageDataFormat = UE::Concert::Compression::GetCompressionDetails(PackageDataStream.DataSize);
 	if (WillUseMemoryWriter(PackageDataStream.DataSize))
 	{
 		DstAr = &MemWriter;
@@ -448,9 +448,9 @@ WritePackageResult WritePackage(FConcertPackageDataStream& PackageDataStream, FA
 	check(PackageDataStream.DataAr != nullptr || (PackageDataStream.DataAr == nullptr && PackageDataStream.DataSize == 0 && PackageDataStream.DataBlob == nullptr));
 
 	WritePackageHeader(*DstAr, PackageDataFormat);
-	if (PackageDataFormat == EPackageDataFormat::PackageDataCompressed)
+	if (UE::Concert::Compression::DataIsCompressed(PackageDataFormat))
 	{
-		WritePackageCompressed(PackageDataStream, *DstAr);
+		WritePackageCompressed(PackageDataStream, *DstAr, PackageDataFormat);
 		OutResult.bShouldCache = ShouldCachePackageBlob(OutResult.CacheData.Num());
 	}
 	else
@@ -473,15 +473,15 @@ bool ExtractPackageData(FArchive& PackageBlobAr, const TFunctionRef<void(FConcer
 	}
 
 	// Read the header.
-	TTuple<EPackageDataFormat, int64> Header = ReadPackageHeader(PackageBlobAr);
-	EPackageDataFormat PackageFormat = Header.Get<0>();
+	TTuple<EConcertCompressionDetails, int64> Header = ReadPackageHeader(PackageBlobAr);
+	EConcertCompressionDetails PackageFormat = Header.Get<0>();
 	int64 PackageDataOffset = Header.Get<1>();
 
 	// Put the stream in position.
 	PackageBlobAr.Seek(PackageDataOffset);
 
 	// The package data is compressed (this hints the package is small enough to be loaded in memory)
-	if (PackageFormat == EPackageDataFormat::PackageDataCompressed)
+	if (UE::Concert::Compression::DataIsCompressed(PackageFormat))
 	{
 		// Read the uncompressed package data size.
 		uint32 UncompressedPackageSize = 0;
@@ -493,7 +493,9 @@ bool ExtractPackageData(FArchive& PackageBlobAr, const TFunctionRef<void(FConcer
 		if (UncompressedPackageSize > 0)
 		{
 			DecompressedData.AddUninitialized(UncompressedPackageSize);
-			PackageBlobAr.SerializeCompressed(DecompressedData.GetData(), UncompressedPackageSize, NAME_Zlib); // Read and decompress the data.
+			PackageBlobAr.SerializeCompressed(DecompressedData.GetData(), UncompressedPackageSize,
+											  UE::Concert::Compression::GetCompressionAlgorithm(PackageFormat),
+											  UE::Concert::Compression::GetCoreCompressionFlags(PackageFormat));
 		}
 
 		// Let the caller stream the package data from a memory archive.
@@ -503,7 +505,7 @@ bool ExtractPackageData(FArchive& PackageBlobAr, const TFunctionRef<void(FConcer
 	}
 	else
 	{
-		check(PackageFormat == EPackageDataFormat::PackageDataUncompressed);
+		check(UE::Concert::Compression::DataIsUncompressed(PackageFormat));
 
 		// Read the package size.
 		int64 PackageDataSize = 0;
@@ -533,13 +535,13 @@ bool ExtractPackageData(const FString& InPackageBlobPathname, const TFunctionRef
 template <typename SerializedDataT>
 void SetSerializedPayloadFlags(SerializedDataT& Data, int32 Flags)
 {
-	Data.bPayloadIsCompressed = static_cast<bool>(Flags);
+	Data.PayloadCompressionDetails = static_cast<EConcertCompressionDetails>(Flags);
 }
 
 template <typename SerializedDataT>
 int32 ConvertSerializedPayloadFlagsToInt32(const SerializedDataT& Data)
 {
-	return static_cast<int32>(Data.bPayloadIsCompressed);
+	return static_cast<int32>(Data.PayloadCompressionDetails);
 }
 
 } // namespace PackageDataUtil
diff --git a/Engine/Plugins/Developer/Concert/ConcertSync/ConcertSyncServer/Source/ConcertSyncServer/Private/ConcertServerDataStore.cpp b/Engine/Plugins/Developer/Concert/ConcertSync/ConcertSyncServer/Source/ConcertSyncServer/Private/ConcertServerDataStore.cpp
--- a/Engine/Plugins/Developer/Concert/ConcertSync/ConcertSyncServer/Source/ConcertSyncServer/Private/ConcertServerDataStore.cpp
+++ b/Engine/Plugins/Developer/Concert/ConcertSync/ConcertSyncServer/Source/ConcertSyncServer/Private/ConcertServerDataStore.cpp
@@ -67,7 +67,7 @@ EConcertSessionResponseCode FConcertServerDataStore::OnCompareExchange(const FCo
 		// If the version or value matches.
 		if ((Request.ExpectedVersion != 0 && Request.ExpectedVersion == Result.Value->Version) ||
 		    (Request.ExpectedVersion == 0
-			 && Result.Value->SerializedValue.bPayloadIsCompressed == Request.Expected.bPayloadIsCompressed
+			 && Result.Value->SerializedValue.PayloadCompressionDetails == Request.Expected.PayloadCompressionDetails
 			 && Result.Value->SerializedValue.PayloadSize == Request.Expected.PayloadSize
 			 && Result.Value->SerializedValue.PayloadBytes.Bytes == Request.Expected.PayloadBytes.Bytes))
 		{
