[Engine/Plugins]
* Ran IWYU on ~170 plugins to remove includes not needed. Public api still keep old includes inside #if UE_ENABLE_INCLUDE_ORDER_DEPRECATED_IN_5_2

diff --git a/Engine/Plugins/2D/Paper2D/Source/Paper2D/Public/Paper2DModule.h b/Engine/Plugins/2D/Paper2D/Source/Paper2D/Public/Paper2DModule.h
--- a/Engine/Plugins/2D/Paper2D/Source/Paper2D/Public/Paper2DModule.h
+++ b/Engine/Plugins/2D/Paper2D/Source/Paper2D/Public/Paper2DModule.h
@@ -6,6 +6,8 @@
 #include "Modules/ModuleInterface.h"
 #include "Stats/Stats2.h"
 
+enum class EStatFlags : uint8;
+
 #define PAPER2D_MODULE_NAME "Paper2D"
 
 //////////////////////////////////////////////////////////////////////////
diff --git a/Engine/Plugins/Animation/AnimationModifierLibrary/Source/AnimationModifierLibrary/Public/AnimationModifierLibrary.h b/Engine/Plugins/Animation/AnimationModifierLibrary/Source/AnimationModifierLibrary/Public/AnimationModifierLibrary.h
--- a/Engine/Plugins/Animation/AnimationModifierLibrary/Source/AnimationModifierLibrary/Public/AnimationModifierLibrary.h
+++ b/Engine/Plugins/Animation/AnimationModifierLibrary/Source/AnimationModifierLibrary/Public/AnimationModifierLibrary.h
@@ -2,8 +2,7 @@
 
 #pragma once
 
-#include "CoreMinimal.h"
-#include "Modules/ModuleManager.h"
+#include "Modules/ModuleInterface.h"
 
 class FAnimationModifierLibraryModule : public IModuleInterface
 {
@@ -12,4 +11,9 @@ public:
 	/** IModuleInterface implementation */
 	virtual void StartupModule() override;
 	virtual void ShutdownModule() override;
-};
\ No newline at end of file
+};
+
+#if UE_ENABLE_INCLUDE_ORDER_DEPRECATED_IN_5_2
+#include "CoreMinimal.h"
+#include "Modules/ModuleManager.h"
+#endif
diff --git a/Engine/Plugins/Animation/AnimationModifierLibrary/Source/AnimationModifierLibrary/Public/CopyBonesModifier.h b/Engine/Plugins/Animation/AnimationModifierLibrary/Source/AnimationModifierLibrary/Public/CopyBonesModifier.h
--- a/Engine/Plugins/Animation/AnimationModifierLibrary/Source/AnimationModifierLibrary/Public/CopyBonesModifier.h
+++ b/Engine/Plugins/Animation/AnimationModifierLibrary/Source/AnimationModifierLibrary/Public/CopyBonesModifier.h
@@ -2,9 +2,9 @@
 
 #pragma once
 
-#include "CoreMinimal.h"
 #include "AnimationModifier.h"
 #include "AnimPose.h"
+#include "BoneContainer.h"
 #include "CopyBonesModifier.generated.h"
 
 USTRUCT(BlueprintType)
@@ -43,3 +43,7 @@ public:
 	virtual void OnRevert_Implementation(UAnimSequence* Animation) override;
 
 };
+
+#if UE_ENABLE_INCLUDE_ORDER_DEPRECATED_IN_5_2
+#include "CoreMinimal.h"
+#endif
diff --git a/Engine/Plugins/Animation/AnimationModifierLibrary/Source/AnimationModifierLibrary/Public/MotionExtractorModifier.h b/Engine/Plugins/Animation/AnimationModifierLibrary/Source/AnimationModifierLibrary/Public/MotionExtractorModifier.h
--- a/Engine/Plugins/Animation/AnimationModifierLibrary/Source/AnimationModifierLibrary/Public/MotionExtractorModifier.h
+++ b/Engine/Plugins/Animation/AnimationModifierLibrary/Source/AnimationModifierLibrary/Public/MotionExtractorModifier.h
@@ -2,11 +2,13 @@
 
 #pragma once
 
-#include "CoreMinimal.h"
 #include "AnimationModifier.h"
-#include "MotionExtractorTypes.h"
 #include "MotionExtractorModifier.generated.h"
 
+enum class EMotionExtractor_Axis : uint8;
+enum class EMotionExtractor_MathOperation : uint8;
+enum class EMotionExtractor_MotionType : uint8;
+
 /** Extracts motion from a bone in the animation and bakes it into a curve */
 UCLASS()
 class UMotionExtractorModifier : public UAnimationModifier
@@ -80,3 +82,8 @@ public:
 	/** Returns the desired value from the extracted poses */
 	float GetDesiredValue(const FTransform& BoneTransform, const FTransform& LastBoneTransform, float DeltaTime) const;
 };
+
+#if UE_ENABLE_INCLUDE_ORDER_DEPRECATED_IN_5_2
+#include "CoreMinimal.h"
+#include "MotionExtractorTypes.h"
+#endif
diff --git a/Engine/Plugins/Animation/AnimationModifierLibrary/Source/AnimationModifierLibrary/Public/MotionExtractorTypes.h b/Engine/Plugins/Animation/AnimationModifierLibrary/Source/AnimationModifierLibrary/Public/MotionExtractorTypes.h
--- a/Engine/Plugins/Animation/AnimationModifierLibrary/Source/AnimationModifierLibrary/Public/MotionExtractorTypes.h
+++ b/Engine/Plugins/Animation/AnimationModifierLibrary/Source/AnimationModifierLibrary/Public/MotionExtractorTypes.h
@@ -2,8 +2,6 @@
 
 #pragma once
 
-#include "CoreMinimal.h"
-#include "AnimationBlueprintLibrary.h"
 #include "MotionExtractorTypes.generated.h"
 
 /** Type of motion to extract */
@@ -39,4 +37,9 @@ enum class EMotionExtractor_MathOperation : uint8
 	Subtraction,
 	Division,
 	Multiplication
-};
\ No newline at end of file
+};
+
+#if UE_ENABLE_INCLUDE_ORDER_DEPRECATED_IN_5_2
+#include "AnimationBlueprintLibrary.h"
+#include "CoreMinimal.h"
+#endif
diff --git a/Engine/Plugins/Animation/AnimationModifierLibrary/Source/AnimationModifierLibrary/Public/MotionExtractorUtilities.h b/Engine/Plugins/Animation/AnimationModifierLibrary/Source/AnimationModifierLibrary/Public/MotionExtractorUtilities.h
--- a/Engine/Plugins/Animation/AnimationModifierLibrary/Source/AnimationModifierLibrary/Public/MotionExtractorUtilities.h
+++ b/Engine/Plugins/Animation/AnimationModifierLibrary/Source/AnimationModifierLibrary/Public/MotionExtractorUtilities.h
@@ -3,9 +3,15 @@
 #pragma once
 
 #include "Kismet/BlueprintFunctionLibrary.h"
-#include "MotionExtractorTypes.h"
+#include "UObject/Package.h"
 #include "MotionExtractorUtilities.generated.h"
 
+class UAnimSequence;
+enum class EMotionExtractor_Axis : uint8;
+enum class EMotionExtractor_MotionType : uint8;
+struct FBoneContainer;
+struct FCompactPoseBoneIndex;
+
 UCLASS(meta=(ScriptName="MotionExtractorUtilityLibrary"))
 class  UMotionExtractorUtilityLibrary : public UBlueprintFunctionLibrary
 {
@@ -71,4 +77,8 @@ public:
 
 	/** Helper function to calculate the magnitude of a vector only considering a specific axis or axes */
 	static float CalculateMagnitude(const FVector& Vector, EMotionExtractor_Axis Axis);
-};
\ No newline at end of file
+};
+
+#if UE_ENABLE_INCLUDE_ORDER_DEPRECATED_IN_5_2
+#include "MotionExtractorTypes.h"
+#endif
diff --git a/Engine/Plugins/Animation/AnimationModifierLibrary/Source/AnimationModifierLibrary/Public/ReOrientRootBoneModifier.h b/Engine/Plugins/Animation/AnimationModifierLibrary/Source/AnimationModifierLibrary/Public/ReOrientRootBoneModifier.h
--- a/Engine/Plugins/Animation/AnimationModifierLibrary/Source/AnimationModifierLibrary/Public/ReOrientRootBoneModifier.h
+++ b/Engine/Plugins/Animation/AnimationModifierLibrary/Source/AnimationModifierLibrary/Public/ReOrientRootBoneModifier.h
@@ -2,7 +2,6 @@
 
 #pragma once
 
-#include "CoreMinimal.h"
 #include "AnimationModifier.h"
 #include "ReOrientRootBoneModifier.generated.h"
 
@@ -28,3 +27,7 @@ private:
 	void ReOrientRootBone_Internal(UAnimSequence* Animation, const FQuat& Quat);
 
 };
+
+#if UE_ENABLE_INCLUDE_ORDER_DEPRECATED_IN_5_2
+#include "CoreMinimal.h"
+#endif
diff --git a/Engine/Plugins/Animation/LiveLink/Source/LiveLink/Private/LiveLinkLogInstance.cpp b/Engine/Plugins/Animation/LiveLink/Source/LiveLink/Private/LiveLinkLogInstance.cpp
--- a/Engine/Plugins/Animation/LiveLink/Source/LiveLink/Private/LiveLinkLogInstance.cpp
+++ b/Engine/Plugins/Animation/LiveLink/Source/LiveLink/Private/LiveLinkLogInstance.cpp
@@ -1,10 +1,9 @@
 // Copyright Epic Games, Inc. All Rights Reserved.
 
 #include "LiveLinkLogInstance.h"
-
+#include "HAL/IConsoleManager.h"
 #include "LiveLinkClient.h"
 
-
 #if WITH_EDITOR
 #include "IMessageLogListing.h"
 #include "MessageLogModule.h"
diff --git a/Engine/Plugins/Animation/LiveLink/Source/LiveLink/Public/LiveLinkAnimationVirtualSubject.h b/Engine/Plugins/Animation/LiveLink/Source/LiveLink/Public/LiveLinkAnimationVirtualSubject.h
--- a/Engine/Plugins/Animation/LiveLink/Source/LiveLink/Public/LiveLinkAnimationVirtualSubject.h
+++ b/Engine/Plugins/Animation/LiveLink/Source/LiveLink/Public/LiveLinkAnimationVirtualSubject.h
@@ -4,13 +4,12 @@
 
 #include "LiveLinkVirtualSubject.h"
 
-#include "CoreMinimal.h"
-#include "ILiveLinkClient.h"
-#include "LiveLinkRefSkeleton.h"
-#include "LiveLinkTypes.h"
 
 #include "LiveLinkAnimationVirtualSubject.generated.h"
 
+struct FLiveLinkSubjectFrameData;
+struct FLiveLinkSubjectKey;
+
 
 // A Skeleton virtual subject is an assembly of different subjects supporting the animation role
 UCLASS(meta=(DisplayName="Animation Virtual Subject"))
@@ -50,3 +49,8 @@ public:
 	UPROPERTY(EditAnywhere, Category = "Settings")
 	bool bAppendSubjectNameToBones;
 };
+
+#if UE_ENABLE_INCLUDE_ORDER_DEPRECATED_IN_5_2
+#include "CoreMinimal.h"
+#include "ILiveLinkClient.h"
+#endif
diff --git a/Engine/Plugins/Animation/LiveLink/Source/LiveLink/Public/LiveLinkBlueprintLibrary.h b/Engine/Plugins/Animation/LiveLink/Source/LiveLink/Public/LiveLinkBlueprintLibrary.h
--- a/Engine/Plugins/Animation/LiveLink/Source/LiveLink/Public/LiveLinkBlueprintLibrary.h
+++ b/Engine/Plugins/Animation/LiveLink/Source/LiveLink/Public/LiveLinkBlueprintLibrary.h
@@ -2,15 +2,26 @@
 
 #pragma once
 
-#include "CoreMinimal.h"
 
 #include "Kismet/BlueprintFunctionLibrary.h"
-#include "LiveLinkRole.h"
-#include "Roles/LiveLinkAnimationBlueprintStructs.h"
-#include "Roles/LiveLinkBasicTypes.h"
 
 #include "LiveLinkBlueprintLibrary.generated.h"
 
+class ULiveLinkRole;
+struct FLiveLinkAnimationFrameData;
+struct FLiveLinkBaseBlueprintData;
+struct FLiveLinkBasicBlueprintData;
+struct FLiveLinkSkeletonStaticData;
+struct FLiveLinkSourceHandle;
+struct FLiveLinkSubjectKey;
+struct FLiveLinkSubjectName;
+struct FLiveLinkSubjectRepresentation;
+struct FLiveLinkTransform;
+struct FSubjectFrameHandle;
+struct FSubjectMetadata;
+struct FTimecode;
+template <typename T> class TSubclassOf;
+
 UCLASS()
 class LIVELINK_API ULiveLinkBlueprintLibrary : public UBlueprintFunctionLibrary
 {
@@ -193,4 +204,11 @@ private:
 	DECLARE_FUNCTION(execEvaluateLiveLinkFrameWithSpecificRole);
 	DECLARE_FUNCTION(execEvaluateLiveLinkFrameAtWorldTimeOffset);
 	DECLARE_FUNCTION(execEvaluateLiveLinkFrameAtSceneTime);
-};
\ No newline at end of file
+};
+
+#if UE_ENABLE_INCLUDE_ORDER_DEPRECATED_IN_5_2
+#include "CoreMinimal.h"
+#include "LiveLinkRole.h"
+#include "Roles/LiveLinkAnimationBlueprintStructs.h"
+#include "Roles/LiveLinkBasicTypes.h"
+#endif
diff --git a/Engine/Plugins/Animation/LiveLink/Source/LiveLink/Public/LiveLinkClient.h b/Engine/Plugins/Animation/LiveLink/Source/LiveLink/Public/LiveLinkClient.h
--- a/Engine/Plugins/Animation/LiveLink/Source/LiveLink/Public/LiveLinkClient.h
+++ b/Engine/Plugins/Animation/LiveLink/Source/LiveLink/Public/LiveLinkClient.h
@@ -4,19 +4,10 @@
 
 #include "ILiveLinkClient.h"
 
-#include "Algo/Transform.h"
-#include "HAL/CriticalSection.h"
-#include "ILiveLinkSource.h"
-#include "LiveLinkSourceSettings.h"
-#include "LiveLinkRefSkeleton.h"
-#include "LiveLinkTypes.h"
-#include "LiveLinkVirtualSubject.h"
-#include "LiveLinkProvider.h"
-#include "Misc/Optional.h"
-#include "Misc/QualifiedFrameTime.h"
-#include "Misc/Timecode.h"
-#include "Tickable.h"
-#include "UObject/GCObject.h"
+
+class ILiveLinkSubject;
+struct FPropertyChangedEvent;
+struct ILiveLinkProvider;
 
 class ULiveLinkPreset;
 class ULiveLinkSourceSettings;
@@ -302,3 +293,12 @@ private:
 	TOptional<FQualifiedFrameTime> CachedEngineFrameTime;
 #endif
 };
+
+#if UE_ENABLE_INCLUDE_ORDER_DEPRECATED_IN_5_2
+#include "ILiveLinkSource.h"
+#include "LiveLinkProvider.h"
+#include "LiveLinkSourceSettings.h"
+#include "LiveLinkVirtualSubject.h"
+#include "Tickable.h"
+#include "UObject/GCObject.h"
+#endif
diff --git a/Engine/Plugins/Animation/LiveLink/Source/LiveLink/Public/LiveLinkClientReference.h b/Engine/Plugins/Animation/LiveLink/Source/LiveLink/Public/LiveLinkClientReference.h
--- a/Engine/Plugins/Animation/LiveLink/Source/LiveLink/Public/LiveLinkClientReference.h
+++ b/Engine/Plugins/Animation/LiveLink/Source/LiveLink/Public/LiveLinkClientReference.h
@@ -2,7 +2,7 @@
 
 #pragma once
 
-#include "CoreMinimal.h"
+#include "HAL/Platform.h"
 
 class ILiveLinkClient;
 
@@ -11,4 +11,8 @@ struct LIVELINK_API FLiveLinkClientReference
 {
 public:
 	ILiveLinkClient* GetClient() const;
-};
\ No newline at end of file
+};
+
+#if UE_ENABLE_INCLUDE_ORDER_DEPRECATED_IN_5_2
+#include "CoreMinimal.h"
+#endif
diff --git a/Engine/Plugins/Animation/LiveLink/Source/LiveLink/Public/LiveLinkComponent.h b/Engine/Plugins/Animation/LiveLink/Source/LiveLink/Public/LiveLinkComponent.h
--- a/Engine/Plugins/Animation/LiveLink/Source/LiveLink/Public/LiveLinkComponent.h
+++ b/Engine/Plugins/Animation/LiveLink/Source/LiveLink/Public/LiveLinkComponent.h
@@ -3,12 +3,12 @@
 #pragma once
 
 #include "Components/ActorComponent.h"
-#include "CoreMinimal.h"
-#include "ILiveLinkClient.h"
-#include "LiveLinkTypes.h"
-#include "Roles/LiveLinkAnimationBlueprintStructs.h"
 #include "LiveLinkComponent.generated.h"
 
+class ILiveLinkClient;
+struct FSubjectFrameHandle;
+struct FTimecode;
+
 DECLARE_DYNAMIC_MULTICAST_DELEGATE_OneParam(FLiveLinkTickSignature, float, DeltaTime);
 
 // An actor component to enable accessing LiveLink data in Blueprints. 
@@ -67,4 +67,11 @@ private:
 	bool bIsDirty;
 
 	ILiveLinkClient* LiveLinkClient;
-};
\ No newline at end of file
+};
+
+#if UE_ENABLE_INCLUDE_ORDER_DEPRECATED_IN_5_2
+#include "CoreMinimal.h"
+#include "ILiveLinkClient.h"
+#include "LiveLinkTypes.h"
+#include "Roles/LiveLinkAnimationBlueprintStructs.h"
+#endif
diff --git a/Engine/Plugins/Animation/LiveLink/Source/LiveLink/Public/LiveLinkMessageBusFinder.h b/Engine/Plugins/Animation/LiveLink/Source/LiveLink/Public/LiveLinkMessageBusFinder.h
--- a/Engine/Plugins/Animation/LiveLink/Source/LiveLink/Public/LiveLinkMessageBusFinder.h
+++ b/Engine/Plugins/Animation/LiveLink/Source/LiveLink/Public/LiveLinkMessageBusFinder.h
@@ -2,15 +2,14 @@
 
 #pragma once
 
-#include "CoreMinimal.h"
-#include "ILiveLinkSource.h"
-#include "LiveLinkTypes.h"
-#include "Misc/ScopeLock.h"
-#include "MessageEndpoint.h"
+#include "IMessageContext.h"
 #include "LatentActions.h"
 #include "Engine/LatentActionManager.h"
 #include "LiveLinkMessageBusFinder.generated.h"
 
+class FMessageEndpoint;
+struct FLiveLinkSourceHandle;
+
 struct FLiveLinkPongMessage;
 
 
@@ -157,4 +156,12 @@ public:
 		return FString::Printf(TEXT("Searching for LiveLink Message Bus providers."));
 	}
 #endif
-};
\ No newline at end of file
+};
+
+#if UE_ENABLE_INCLUDE_ORDER_DEPRECATED_IN_5_2
+#include "CoreMinimal.h"
+#include "ILiveLinkSource.h"
+#include "LiveLinkTypes.h"
+#include "MessageEndpoint.h"
+#include "Misc/ScopeLock.h"
+#endif
diff --git a/Engine/Plugins/Animation/LiveLink/Source/LiveLink/Public/LiveLinkMessageBusSource.h b/Engine/Plugins/Animation/LiveLink/Source/LiveLink/Public/LiveLinkMessageBusSource.h
--- a/Engine/Plugins/Animation/LiveLink/Source/LiveLink/Public/LiveLinkMessageBusSource.h
+++ b/Engine/Plugins/Animation/LiveLink/Source/LiveLink/Public/LiveLinkMessageBusSource.h
@@ -5,10 +5,12 @@
 #include "ILiveLinkSource.h"
 
 #include "HAL/ThreadSafeBool.h"
-#include "IMessageContext.h"
-#include "LiveLinkRole.h"
 #include "MessageEndpoint.h"
 
+class ULiveLinkRole;
+class ULiveLinkSourceSettings;
+struct FMessageEndpointBuilder;
+
 class ILiveLinkClient;
 struct FLiveLinkPongMessage;
 struct FLiveLinkSubjectDataMessage;
@@ -119,4 +121,8 @@ private:
 
 	// Offset between sender's machine engine time and receiver's machine engine time
 	double MachineTimeOffset;
-};
\ No newline at end of file
+};
+
+#if UE_ENABLE_INCLUDE_ORDER_DEPRECATED_IN_5_2
+#include "LiveLinkRole.h"
+#endif
diff --git a/Engine/Plugins/Animation/LiveLink/Source/LiveLink/Public/LiveLinkMessageBusSourceFactory.h b/Engine/Plugins/Animation/LiveLink/Source/LiveLink/Public/LiveLinkMessageBusSourceFactory.h
--- a/Engine/Plugins/Animation/LiveLink/Source/LiveLink/Public/LiveLinkMessageBusSourceFactory.h
+++ b/Engine/Plugins/Animation/LiveLink/Source/LiveLink/Public/LiveLinkMessageBusSourceFactory.h
@@ -6,6 +6,8 @@
 #include "LiveLinkSourceFactory.h"
 #include "LiveLinkMessageBusSourceFactory.generated.h"
 
+struct FMessageAddress;
+
 
 UCLASS()
 class LIVELINK_API ULiveLinkMessageBusSourceFactory : public ULiveLinkSourceFactory
diff --git a/Engine/Plugins/Animation/LiveLink/Source/LiveLink/Public/LiveLinkPreset.h b/Engine/Plugins/Animation/LiveLink/Source/LiveLink/Public/LiveLinkPreset.h
--- a/Engine/Plugins/Animation/LiveLink/Source/LiveLink/Public/LiveLinkPreset.h
+++ b/Engine/Plugins/Animation/LiveLink/Source/LiveLink/Public/LiveLinkPreset.h
@@ -2,16 +2,13 @@
 
 #pragma once
 
-#include "CoreMinimal.h"
-#include "Engine/EngineTypes.h"
-#include "Engine/LatentActionManager.h"
-#include "LiveLinkPresetTypes.h"
-#include "Templates/PimplPtr.h"
-#include "UObject/Object.h"
-#include "UObject/ObjectMacros.h"
 
 #include "LiveLinkPreset.generated.h"
 
+struct FLatentActionInfo;
+struct FLiveLinkSourcePreset;
+struct FLiveLinkSubjectPreset;
+
 UCLASS(BlueprintType)
 class LIVELINK_API ULiveLinkPreset : public UObject
 {
@@ -74,3 +71,10 @@ private:
 	/** Holds the current ApplyToClient async operation. Only one operation for all presets can be done at a time. */
 	static TPimplPtr<struct FApplyToClientPollingOperation> ApplyToClientPollingOperation;
 };
+
+#if UE_ENABLE_INCLUDE_ORDER_DEPRECATED_IN_5_2
+#include "CoreMinimal.h"
+#include "Engine/EngineTypes.h"
+#include "Engine/LatentActionManager.h"
+#include "LiveLinkPresetTypes.h"
+#endif
diff --git a/Engine/Plugins/Animation/LiveLink/Source/LiveLink/Public/LiveLinkRetargetAssetReference.h b/Engine/Plugins/Animation/LiveLink/Source/LiveLink/Public/LiveLinkRetargetAssetReference.h
--- a/Engine/Plugins/Animation/LiveLink/Source/LiveLink/Public/LiveLinkRetargetAssetReference.h
+++ b/Engine/Plugins/Animation/LiveLink/Source/LiveLink/Public/LiveLinkRetargetAssetReference.h
@@ -2,7 +2,6 @@
 
 #pragma once
 
-#include "CoreMinimal.h"
 
 #include "LiveLinkRetargetAssetReference.generated.h"
 
@@ -16,4 +15,8 @@ struct
 {
 public:
 	GENERATED_BODY()
-};
\ No newline at end of file
+};
+
+#if UE_ENABLE_INCLUDE_ORDER_DEPRECATED_IN_5_2
+#include "CoreMinimal.h"
+#endif
diff --git a/Engine/Plugins/Animation/LiveLink/Source/LiveLink/Public/LiveLinkSettings.h b/Engine/Plugins/Animation/LiveLink/Source/LiveLink/Public/LiveLinkSettings.h
--- a/Engine/Plugins/Animation/LiveLink/Source/LiveLink/Public/LiveLinkSettings.h
+++ b/Engine/Plugins/Animation/LiveLink/Source/LiveLink/Public/LiveLinkSettings.h
@@ -2,20 +2,17 @@
 
 #pragma once
 
-#include "CoreTypes.h"
-#include "UObject/Object.h"
-#include "UObject/ObjectMacros.h"
 
-#include "Engine/EngineTypes.h"
-#include "LiveLinkFrameInterpolationProcessor.h"
-#include "LiveLinkFramePreProcessor.h"
-#include "LiveLinkRole.h"
-#include "LiveLinkSourceSettings.h"
-#include "LiveLinkSubjectSettings.h"
 #include "Templates/SubclassOf.h"
 
 #include "LiveLinkSettings.generated.h"
 
+class ULiveLinkFrameInterpolationProcessor;
+class ULiveLinkFramePreProcessor;
+class ULiveLinkRole;
+class ULiveLinkSubjectSettings;
+enum class ELiveLinkSourceMode : uint8;
+
 
 class ULiveLinkPreset;
 
@@ -159,3 +156,12 @@ public:
 	double GetMessageBusHeartbeatTimeout() const { return MessageBusHeartbeatTimeout; }
 	double GetMessageBusTimeBeforeRemovingDeadSource() const { return MessageBusTimeBeforeRemovingInactiveSource; }
 };
+
+#if UE_ENABLE_INCLUDE_ORDER_DEPRECATED_IN_5_2
+#include "Engine/EngineTypes.h"
+#include "LiveLinkFrameInterpolationProcessor.h"
+#include "LiveLinkFramePreProcessor.h"
+#include "LiveLinkRole.h"
+#include "LiveLinkSourceSettings.h"
+#include "LiveLinkSubjectSettings.h"
+#endif
diff --git a/Engine/Plugins/Animation/LiveLink/Source/LiveLink/Public/LiveLinkTimeSynchronizationSource.h b/Engine/Plugins/Animation/LiveLink/Source/LiveLink/Public/LiveLinkTimeSynchronizationSource.h
--- a/Engine/Plugins/Animation/LiveLink/Source/LiveLink/Public/LiveLinkTimeSynchronizationSource.h
+++ b/Engine/Plugins/Animation/LiveLink/Source/LiveLink/Public/LiveLinkTimeSynchronizationSource.h
@@ -2,12 +2,12 @@
 
 #pragma once
 
-#include "CoreMinimal.h"
 #include "TimeSynchronizationSource.h"
 #include "LiveLinkClient.h"
-#include "Misc/Guid.h"
 #include "LiveLinkTimeSynchronizationSource.generated.h"
 
+struct FFrameRate;
+
 UCLASS(EditInlineNew)
 class LIVELINK_API ULiveLinkTimeSynchronizationSource : public UTimeSynchronizationSource
 {
@@ -52,4 +52,8 @@ private:
 	void OnModularFeatureRegistered(const FName& FeatureName, class IModularFeature* Feature);
 	void OnModularFeatureUnregistered(const FName& FeatureName, class IModularFeature* Feature);
 	void UpdateCachedState() const;
-};
\ No newline at end of file
+};
+
+#if UE_ENABLE_INCLUDE_ORDER_DEPRECATED_IN_5_2
+#include "CoreMinimal.h"
+#endif
diff --git a/Engine/Plugins/Animation/LiveLink/Source/LiveLink/Public/LiveLinkTimecodeProvider.h b/Engine/Plugins/Animation/LiveLink/Source/LiveLink/Public/LiveLinkTimecodeProvider.h
--- a/Engine/Plugins/Animation/LiveLink/Source/LiveLink/Public/LiveLinkTimecodeProvider.h
+++ b/Engine/Plugins/Animation/LiveLink/Source/LiveLink/Public/LiveLinkTimecodeProvider.h
@@ -4,11 +4,12 @@
 
 #include "Engine/TimecodeProvider.h"
 
-#include "HAL/CriticalSection.h"
 #include "LiveLinkTypes.h"
-#include "Templates/Atomic.h"
 #include "LiveLinkTimecodeProvider.generated.h"
 
+struct FPropertyChangedEvent;
+struct FQualifiedFrameTime;
+
 class ILiveLinkClient;
 class IModularFeature;
 
diff --git a/Engine/Plugins/Animation/LiveLink/Source/LiveLink/Public/VirtualSubjects/LiveLinkBlueprintVirtualSubject.h b/Engine/Plugins/Animation/LiveLink/Source/LiveLink/Public/VirtualSubjects/LiveLinkBlueprintVirtualSubject.h
--- a/Engine/Plugins/Animation/LiveLink/Source/LiveLink/Public/VirtualSubjects/LiveLinkBlueprintVirtualSubject.h
+++ b/Engine/Plugins/Animation/LiveLink/Source/LiveLink/Public/VirtualSubjects/LiveLinkBlueprintVirtualSubject.h
@@ -2,9 +2,6 @@
 
 #pragma once
 
-#include "CoreMinimal.h"
-#include "ILiveLinkClient.h"
-#include "LiveLinkTypes.h"
 #include "LiveLinkVirtualSubject.h"
 
 #include "LiveLinkBlueprintVirtualSubject.generated.h"
@@ -49,3 +46,8 @@ private:
 	DECLARE_FUNCTION(execUpdateVirtualSubjectStaticData_Internal);
 	DECLARE_FUNCTION(execUpdateVirtualSubjectFrameData_Internal);
 };
+
+#if UE_ENABLE_INCLUDE_ORDER_DEPRECATED_IN_5_2
+#include "CoreMinimal.h"
+#include "ILiveLinkClient.h"
+#endif
diff --git a/Engine/Plugins/Animation/LiveLink/Source/LiveLinkComponents/Public/Controllers/LiveLinkLightController.h b/Engine/Plugins/Animation/LiveLink/Source/LiveLinkComponents/Public/Controllers/LiveLinkLightController.h
--- a/Engine/Plugins/Animation/LiveLink/Source/LiveLinkComponents/Public/Controllers/LiveLinkLightController.h
+++ b/Engine/Plugins/Animation/LiveLink/Source/LiveLinkComponents/Public/Controllers/LiveLinkLightController.h
@@ -2,9 +2,7 @@
 
 #pragma once
 
-#include "LiveLinkControllerBase.h"
 #include "LiveLinkTransformController.h"
-#include "Engine/EngineTypes.h"
 #include "LiveLinkLightController.generated.h"
 
 
@@ -37,4 +35,4 @@ public:
 	//~ Begin UObject interface
 	virtual void PostLoad() override;
 	//~ End UObject interface
-};
\ No newline at end of file
+};
diff --git a/Engine/Plugins/Animation/LiveLink/Source/LiveLinkComponents/Public/Controllers/LiveLinkTransformController.h b/Engine/Plugins/Animation/LiveLink/Source/LiveLinkComponents/Public/Controllers/LiveLinkTransformController.h
--- a/Engine/Plugins/Animation/LiveLink/Source/LiveLinkComponents/Public/Controllers/LiveLinkTransformController.h
+++ b/Engine/Plugins/Animation/LiveLink/Source/LiveLinkComponents/Public/Controllers/LiveLinkTransformController.h
@@ -2,11 +2,11 @@
 
 #pragma once
 
-#include "Math/MathFwd.h"
 #include "LiveLinkControllerBase.h"
-#include "Engine/EngineTypes.h"
 #include "LiveLinkTransformController.generated.h"
 
+class ULiveLinkRole;
+
 struct FLiveLinkTransformStaticData;
 class USceneComponent;
 
@@ -79,4 +79,4 @@ public:
 	//~ Begin UObject interface
 	virtual void PostLoad() override;
 	//~ End UObject interface
-};
\ No newline at end of file
+};
diff --git a/Engine/Plugins/Animation/LiveLink/Source/LiveLinkComponents/Public/ILiveLinkComponentModule.h b/Engine/Plugins/Animation/LiveLink/Source/LiveLinkComponents/Public/ILiveLinkComponentModule.h
--- a/Engine/Plugins/Animation/LiveLink/Source/LiveLinkComponents/Public/ILiveLinkComponentModule.h
+++ b/Engine/Plugins/Animation/LiveLink/Source/LiveLinkComponents/Public/ILiveLinkComponentModule.h
@@ -2,10 +2,11 @@
 
 #pragma once
 
+#include "Delegates/Delegate.h"
 #include "Modules/ModuleInterface.h"
-#include "Modules/ModuleManager.h"
 
-#include "LiveLinkComponentController.h"
+
+class ULiveLinkComponentController;
 
 DECLARE_MULTICAST_DELEGATE_OneParam(FLiveLinkComponentRegistered, ULiveLinkComponentController*);
 
@@ -22,3 +23,8 @@ public:
 	virtual ~ILiveLinkComponentsModule() { }
 };
 
+
+#if UE_ENABLE_INCLUDE_ORDER_DEPRECATED_IN_5_2
+#include "LiveLinkComponentController.h"
+#include "Modules/ModuleManager.h"
+#endif
diff --git a/Engine/Plugins/Animation/LiveLink/Source/LiveLinkComponents/Public/LiveLinkComponentController.h b/Engine/Plugins/Animation/LiveLink/Source/LiveLinkComponents/Public/LiveLinkComponentController.h
--- a/Engine/Plugins/Animation/LiveLink/Source/LiveLinkComponents/Public/LiveLinkComponentController.h
+++ b/Engine/Plugins/Animation/LiveLink/Source/LiveLinkComponents/Public/LiveLinkComponentController.h
@@ -3,9 +3,12 @@
 #pragma once
 
 #include "Components/ActorComponent.h"
-#include "ILiveLinkClient.h"
+#include "LiveLinkRole.h"
 #include "LiveLinkComponentController.generated.h"
 
+class ULiveLinkComponentController;
+struct FLiveLinkSubjectFrameData;
+
 DECLARE_DYNAMIC_MULTICAST_DELEGATE_OneParam(FLiveLinkTickDelegate, float, DeltaTime);
 DECLARE_DYNAMIC_MULTICAST_DELEGATE(FOnControllerMapUpdatedDelegate);
 
@@ -138,3 +141,7 @@ protected:
 private:
 	FLiveLinkControllersTicked LiveLinkControllersTickedDelegate;
 };
+
+#if UE_ENABLE_INCLUDE_ORDER_DEPRECATED_IN_5_2
+#include "ILiveLinkClient.h"
+#endif
diff --git a/Engine/Plugins/Animation/LiveLink/Source/LiveLinkComponents/Public/LiveLinkComponentSettings.h b/Engine/Plugins/Animation/LiveLink/Source/LiveLinkComponents/Public/LiveLinkComponentSettings.h
--- a/Engine/Plugins/Animation/LiveLink/Source/LiveLinkComponents/Public/LiveLinkComponentSettings.h
+++ b/Engine/Plugins/Animation/LiveLink/Source/LiveLinkComponents/Public/LiveLinkComponentSettings.h
@@ -2,15 +2,14 @@
 
 #pragma once
 
-#include "UObject/Object.h"
-#include "UObject/ObjectMacros.h"
 
-#include "LiveLinkRole.h"
-#include "LiveLinkControllerBase.h"
 #include "Templates/SubclassOf.h"
 
 #include "LiveLinkComponentSettings.generated.h"
 
+class ULiveLinkControllerBase;
+class ULiveLinkRole;
+
 
 /**
  * Settings for LiveLink.
@@ -26,3 +25,8 @@ public:
 	UPROPERTY(config, EditAnywhere, Category = "LiveLink", meta = (AllowAbstract = "false"))
 	TMap<TSubclassOf<ULiveLinkRole>, TSubclassOf<ULiveLinkControllerBase>> DefaultControllerForRole;
 };
+
+#if UE_ENABLE_INCLUDE_ORDER_DEPRECATED_IN_5_2
+#include "LiveLinkControllerBase.h"
+#include "LiveLinkRole.h"
+#endif
diff --git a/Engine/Plugins/Animation/LiveLink/Source/LiveLinkComponents/Public/LiveLinkControllerBase.h b/Engine/Plugins/Animation/LiveLink/Source/LiveLinkComponents/Public/LiveLinkControllerBase.h
--- a/Engine/Plugins/Animation/LiveLink/Source/LiveLinkComponents/Public/LiveLinkControllerBase.h
+++ b/Engine/Plugins/Animation/LiveLink/Source/LiveLinkComponents/Public/LiveLinkControllerBase.h
@@ -2,16 +2,14 @@
 
 #pragma once
 
-#include "CoreMinimal.h"
-#include "UObject/ObjectMacros.h"
-#include "UObject/UObjectGlobals.h"
-#include "UObject/Object.h"
 
 #include "Components/ActorComponent.h"
-#include "ILiveLinkClient.h"
 
+#include "LiveLinkRole.h"
 #include "LiveLinkControllerBase.generated.h"
 
+struct FLiveLinkSubjectFrameData;
+
 class AActor;
 
 /**
@@ -107,3 +105,8 @@ protected:
 	FLiveLinkSubjectRepresentation SelectedSubject;
 };
 
+
+#if UE_ENABLE_INCLUDE_ORDER_DEPRECATED_IN_5_2
+#include "CoreMinimal.h"
+#include "ILiveLinkClient.h"
+#endif
diff --git a/Engine/Plugins/Animation/LiveLink/Source/LiveLinkEditor/Public/SLiveLinkSubjectRepresentationPicker.h b/Engine/Plugins/Animation/LiveLink/Source/LiveLinkEditor/Public/SLiveLinkSubjectRepresentationPicker.h
--- a/Engine/Plugins/Animation/LiveLink/Source/LiveLinkEditor/Public/SLiveLinkSubjectRepresentationPicker.h
+++ b/Engine/Plugins/Animation/LiveLink/Source/LiveLinkEditor/Public/SLiveLinkSubjectRepresentationPicker.h
@@ -2,17 +2,17 @@
 
 #pragma once
 
-#include "Widgets/DeclarativeSyntaxSupport.h"
 #include "Widgets/SCompoundWidget.h"
 
-#include "Input/Reply.h"
 #include "LiveLinkRole.h"
-#include "LiveLinkPreset.h"
 #include "Styling/CoreStyle.h"
 #include "Styling/SlateTypes.h"
 #include "Styling/SlateWidgetStyleAsset.h"
-#include "Types/SlateEnums.h"
-#include "Widgets/Views/SListView.h"
+
+class ULiveLinkPreset;
+namespace ESelectInfo { enum Type : int; }
+namespace ETextCommit { enum Type : int; }
+template <typename ItemType> class SListView;
 
 struct FAssetData;
 struct FLiveLinkSubjectRepresentationPickerEntry;
@@ -144,3 +144,8 @@ private:
 	bool bShowSource;
 	bool bShowRole;
 };
+
+#if UE_ENABLE_INCLUDE_ORDER_DEPRECATED_IN_5_2
+#include "LiveLinkPreset.h"
+#include "Widgets/Views/SListView.h"
+#endif
diff --git a/Engine/Plugins/Animation/LiveLink/Source/LiveLinkEditor/Public/VirtualSubjects/LiveLinkBlueprintVirtualSubjectFactory.h b/Engine/Plugins/Animation/LiveLink/Source/LiveLinkEditor/Public/VirtualSubjects/LiveLinkBlueprintVirtualSubjectFactory.h
--- a/Engine/Plugins/Animation/LiveLink/Source/LiveLinkEditor/Public/VirtualSubjects/LiveLinkBlueprintVirtualSubjectFactory.h
+++ b/Engine/Plugins/Animation/LiveLink/Source/LiveLinkEditor/Public/VirtualSubjects/LiveLinkBlueprintVirtualSubjectFactory.h
@@ -2,9 +2,7 @@
 
 #pragma once
 
-#include "CoreMinimal.h"
 #include "Factories/BlueprintFactory.h"
-#include "UObject/ObjectMacros.h"
 
 #include "LiveLinkBlueprintVirtualSubjectFactory.generated.h"
 
@@ -29,3 +27,7 @@ public:
 	virtual uint32 GetMenuCategories() const override;
 	//~ Begin UFactory Interface
 };
+
+#if UE_ENABLE_INCLUDE_ORDER_DEPRECATED_IN_5_2
+#include "CoreMinimal.h"
+#endif
diff --git a/Engine/Plugins/Animation/LiveLink/Source/LiveLinkGraphNode/Public/AnimGraphNode_LiveLinkPose.h b/Engine/Plugins/Animation/LiveLink/Source/LiveLinkGraphNode/Public/AnimGraphNode_LiveLinkPose.h
--- a/Engine/Plugins/Animation/LiveLink/Source/LiveLinkGraphNode/Public/AnimGraphNode_LiveLinkPose.h
+++ b/Engine/Plugins/Animation/LiveLink/Source/LiveLinkGraphNode/Public/AnimGraphNode_LiveLinkPose.h
@@ -7,6 +7,8 @@
 
 #include "AnimGraphNode_LiveLinkPose.generated.h"
 
+namespace ENodeTitleType { enum Type : int; }
+
 
 
 UCLASS()
diff --git a/Engine/Plugins/Animation/LiveLink/Source/LiveLinkGraphNode/Public/K2Node_EvaluateLiveLinkFrame.h b/Engine/Plugins/Animation/LiveLink/Source/LiveLinkGraphNode/Public/K2Node_EvaluateLiveLinkFrame.h
--- a/Engine/Plugins/Animation/LiveLink/Source/LiveLinkGraphNode/Public/K2Node_EvaluateLiveLinkFrame.h
+++ b/Engine/Plugins/Animation/LiveLink/Source/LiveLinkGraphNode/Public/K2Node_EvaluateLiveLinkFrame.h
@@ -2,15 +2,13 @@
 
 #pragma once
 
-#include "CoreMinimal.h"
-#include "EdGraph/EdGraphNodeUtils.h"
-#include "LiveLinkRole.h"
 #include "K2Node.h"
-#include "Textures/SlateIcon.h"
-#include "UObject/ObjectMacros.h"
 
+#include "UObject/Package.h"
 #include "K2Node_EvaluateLiveLinkFrame.generated.h"
 
+class ULiveLinkRole;
+
 class FBlueprintActionDatabaseRegistrar;
 class FKismetCompilerContext;
 class UDataTable;
@@ -94,3 +92,10 @@ protected:
 	/** Verify if selected role class is valid for evaluation */
 	bool IsRoleValidForEvaluation(TSubclassOf<ULiveLinkRole> InRoleClass) const;
 };
+
+#if UE_ENABLE_INCLUDE_ORDER_DEPRECATED_IN_5_2
+#include "CoreMinimal.h"
+#include "EdGraph/EdGraphNodeUtils.h"
+#include "LiveLinkRole.h"
+#include "Textures/SlateIcon.h"
+#endif
diff --git a/Engine/Plugins/Animation/LiveLink/Source/LiveLinkGraphNode/Public/K2Node_UpdateVirtualSubjectDataBase.h b/Engine/Plugins/Animation/LiveLink/Source/LiveLinkGraphNode/Public/K2Node_UpdateVirtualSubjectDataBase.h
--- a/Engine/Plugins/Animation/LiveLink/Source/LiveLinkGraphNode/Public/K2Node_UpdateVirtualSubjectDataBase.h
+++ b/Engine/Plugins/Animation/LiveLink/Source/LiveLinkGraphNode/Public/K2Node_UpdateVirtualSubjectDataBase.h
@@ -2,12 +2,12 @@
 
 #pragma once
 
-#include "CoreMinimal.h"
 #include "K2Node.h"
-#include "Textures/SlateIcon.h"
 
 #include "K2Node_UpdateVirtualSubjectDataBase.generated.h"
 
+namespace ENodeTitleType { enum Type : int; }
+
 class FBlueprintActionDatabaseRegistrar;
 class FKismetCompilerContext;
 class UDataTable;
@@ -80,3 +80,8 @@ private:
 	/** Struct pin description */
 	static const FText LiveLinkStructPinDescription;
 };
+
+#if UE_ENABLE_INCLUDE_ORDER_DEPRECATED_IN_5_2
+#include "CoreMinimal.h"
+#include "Textures/SlateIcon.h"
+#endif
diff --git a/Engine/Plugins/Animation/LiveLink/Source/LiveLinkMovieScene/Private/MovieScene/MovieSceneLiveLinkSection.cpp b/Engine/Plugins/Animation/LiveLink/Source/LiveLinkMovieScene/Private/MovieScene/MovieSceneLiveLinkSection.cpp
--- a/Engine/Plugins/Animation/LiveLink/Source/LiveLinkMovieScene/Private/MovieScene/MovieSceneLiveLinkSection.cpp
+++ b/Engine/Plugins/Animation/LiveLink/Source/LiveLinkMovieScene/Private/MovieScene/MovieSceneLiveLinkSection.cpp
@@ -3,7 +3,6 @@
 #include "MovieScene/MovieSceneLiveLinkSection.h"
 
 #include "Channels/MovieSceneBoolChannel.h"
-#include "Channels/MovieSceneChannelProxy.h"
 #include "Channels/MovieSceneByteChannel.h"
 #include "LiveLinkCustomVersion.h"
 #include "Channels/MovieSceneFloatChannel.h"
diff --git a/Engine/Plugins/Animation/LiveLink/Source/LiveLinkMovieScene/Private/MovieScene/MovieSceneLiveLinkStructProperties.cpp b/Engine/Plugins/Animation/LiveLink/Source/LiveLinkMovieScene/Private/MovieScene/MovieSceneLiveLinkStructProperties.cpp
--- a/Engine/Plugins/Animation/LiveLink/Source/LiveLinkMovieScene/Private/MovieScene/MovieSceneLiveLinkStructProperties.cpp
+++ b/Engine/Plugins/Animation/LiveLink/Source/LiveLinkMovieScene/Private/MovieScene/MovieSceneLiveLinkStructProperties.cpp
@@ -2,13 +2,19 @@
 
 #include "MovieScene/MovieSceneLiveLinkStructProperties.h"
 
+#include "Channels/MovieSceneBoolChannel.h"
 #include "MovieScene/MovieSceneLiveLinkEnumHandler.h"
+#include "Channels/MovieSceneByteChannel.h"
 #include "MovieScene/MovieSceneLiveLinkPropertyHandler.h"
+#include "Channels/MovieSceneIntegerChannel.h"
 #include "MovieScene/MovieSceneLiveLinkTransformHandler.h"
+#include "Channels/MovieSceneStringChannel.h"
 #include "UObject/EnumProperty.h"
 
 #include UE_INLINE_GENERATED_CPP_BY_NAME(MovieSceneLiveLinkStructProperties)
 
+FLiveLinkPropertyData::FLiveLinkPropertyData() = default;
+FLiveLinkPropertyData::~FLiveLinkPropertyData() = default;
 
 namespace LiveLinkPropertiesUtils
 {
diff --git a/Engine/Plugins/Animation/LiveLink/Source/LiveLinkMovieScene/Private/MovieScene/MovieSceneLiveLinkSubSection.cpp b/Engine/Plugins/Animation/LiveLink/Source/LiveLinkMovieScene/Private/MovieScene/MovieSceneLiveLinkSubSection.cpp
--- a/Engine/Plugins/Animation/LiveLink/Source/LiveLinkMovieScene/Private/MovieScene/MovieSceneLiveLinkSubSection.cpp
+++ b/Engine/Plugins/Animation/LiveLink/Source/LiveLinkMovieScene/Private/MovieScene/MovieSceneLiveLinkSubSection.cpp
@@ -3,15 +3,7 @@
 #include "MovieScene/MovieSceneLiveLinkSubSection.h"
 
 
-#include "Channels/MovieSceneBoolChannel.h"
 #include "UObject/UObjectIterator.h"
-#include "Channels/MovieSceneByteChannel.h"
-#include "Channels/MovieSceneChannel.h"
-#include "Channels/MovieSceneChannelData.h"
-#include "Channels/MovieSceneFloatChannel.h"
-#include "Channels/MovieSceneIntegerChannel.h"
-#include "Channels/MovieSceneStringChannel.h"
-#include "Curves/KeyHandle.h"
 #include "MovieScene/MovieSceneLiveLinkStructProperties.h"
 
 #include UE_INLINE_GENERATED_CPP_BY_NAME(MovieSceneLiveLinkSubSection)
diff --git a/Engine/Plugins/Animation/LiveLink/Source/LiveLinkMovieScene/Private/MovieScene/MovieSceneLiveLinkSubSectionAnimation.cpp b/Engine/Plugins/Animation/LiveLink/Source/LiveLinkMovieScene/Private/MovieScene/MovieSceneLiveLinkSubSectionAnimation.cpp
--- a/Engine/Plugins/Animation/LiveLink/Source/LiveLinkMovieScene/Private/MovieScene/MovieSceneLiveLinkSubSectionAnimation.cpp
+++ b/Engine/Plugins/Animation/LiveLink/Source/LiveLinkMovieScene/Private/MovieScene/MovieSceneLiveLinkSubSectionAnimation.cpp
@@ -2,13 +2,8 @@
 
 #include "MovieScene/MovieSceneLiveLinkSubSectionAnimation.h"
 
-#include "Channels/MovieSceneBoolChannel.h"
-#include "Channels/MovieSceneChannelProxy.h"
-#include "Channels/MovieSceneByteChannel.h"
 #include "Roles/LiveLinkAnimationRole.h"
-#include "Channels/MovieSceneIntegerChannel.h"
 #include "Roles/LiveLinkAnimationTypes.h"
-#include "Channels/MovieSceneStringChannel.h"
 #include "MovieScene/MovieSceneLiveLinkStructProperties.h"
 #include "MovieScene/MovieSceneLiveLinkSubSection.h"
 #include "MovieScene/MovieSceneLiveLinkTransformHandler.h"
diff --git a/Engine/Plugins/Animation/LiveLink/Source/LiveLinkMovieScene/Private/MovieScene/MovieSceneLiveLinkSubSectionBasicRole.cpp b/Engine/Plugins/Animation/LiveLink/Source/LiveLinkMovieScene/Private/MovieScene/MovieSceneLiveLinkSubSectionBasicRole.cpp
--- a/Engine/Plugins/Animation/LiveLink/Source/LiveLinkMovieScene/Private/MovieScene/MovieSceneLiveLinkSubSectionBasicRole.cpp
+++ b/Engine/Plugins/Animation/LiveLink/Source/LiveLinkMovieScene/Private/MovieScene/MovieSceneLiveLinkSubSectionBasicRole.cpp
@@ -2,12 +2,7 @@
 
 #include "MovieScene/MovieSceneLiveLinkSubSectionBasicRole.h"
 
-#include "Channels/MovieSceneBoolChannel.h"
-#include "Channels/MovieSceneChannelProxy.h"
-#include "Channels/MovieSceneByteChannel.h"
 #include "Channels/MovieSceneFloatChannel.h"
-#include "Channels/MovieSceneIntegerChannel.h"
-#include "Channels/MovieSceneStringChannel.h"
 #include "MovieScene/IMovieSceneLiveLinkPropertyHandler.h"
 #include "MovieScene/MovieSceneLiveLinkStructProperties.h"
 #include "MovieScene/MovieSceneLiveLinkSubSection.h"
diff --git a/Engine/Plugins/Animation/LiveLink/Source/LiveLinkMovieScene/Private/MovieScene/MovieSceneLiveLinkSubSectionProperties.cpp b/Engine/Plugins/Animation/LiveLink/Source/LiveLinkMovieScene/Private/MovieScene/MovieSceneLiveLinkSubSectionProperties.cpp
--- a/Engine/Plugins/Animation/LiveLink/Source/LiveLinkMovieScene/Private/MovieScene/MovieSceneLiveLinkSubSectionProperties.cpp
+++ b/Engine/Plugins/Animation/LiveLink/Source/LiveLinkMovieScene/Private/MovieScene/MovieSceneLiveLinkSubSectionProperties.cpp
@@ -3,7 +3,6 @@
 #include "MovieScene/MovieSceneLiveLinkSubSectionProperties.h"
 
 #include "Channels/MovieSceneBoolChannel.h"
-#include "Channels/MovieSceneChannelProxy.h"
 
 #include "Channels/MovieSceneByteChannel.h"
 #include "LiveLinkMovieScenePrivate.h"
diff --git a/Engine/Plugins/Animation/LiveLink/Source/LiveLinkMovieScene/Public/MovieScene/IMovieSceneLiveLinkPropertyHandler.h b/Engine/Plugins/Animation/LiveLink/Source/LiveLinkMovieScene/Public/MovieScene/IMovieSceneLiveLinkPropertyHandler.h
--- a/Engine/Plugins/Animation/LiveLink/Source/LiveLinkMovieScene/Public/MovieScene/IMovieSceneLiveLinkPropertyHandler.h
+++ b/Engine/Plugins/Animation/LiveLink/Source/LiveLinkMovieScene/Public/MovieScene/IMovieSceneLiveLinkPropertyHandler.h
@@ -2,9 +2,15 @@
 
 #pragma once
 
-#include "CoreMinimal.h"
 
-#include "LiveLinkTypes.h"
+#include "Misc/FrameNumber.h"
+
+class UScriptStruct;
+struct FFrameTime;
+struct FLiveLinkBaseFrameData;
+struct FLiveLinkWorldTime;
+struct FQualifiedFrameTime;
+template <typename OptionalType> struct TOptional;
 
 
 struct FKeyDataOptimizationParams;
@@ -35,3 +41,8 @@ public:
 	virtual void FillFrameInterpolated(const FFrameTime& InFrameTime, const FLiveLinkWorldTime& InWorldTime, const TOptional<FQualifiedFrameTime>& InTimecodeTime, const UScriptStruct& InStruct, FLiveLinkBaseFrameData* OutFrame) = 0;
 };
 
+
+#if UE_ENABLE_INCLUDE_ORDER_DEPRECATED_IN_5_2
+#include "CoreMinimal.h"
+#include "LiveLinkTypes.h"
+#endif
diff --git a/Engine/Plugins/Animation/LiveLink/Source/LiveLinkMovieScene/Public/MovieScene/MovieSceneLiveLinkBufferData.h b/Engine/Plugins/Animation/LiveLink/Source/LiveLinkMovieScene/Public/MovieScene/MovieSceneLiveLinkBufferData.h
--- a/Engine/Plugins/Animation/LiveLink/Source/LiveLinkMovieScene/Public/MovieScene/MovieSceneLiveLinkBufferData.h
+++ b/Engine/Plugins/Animation/LiveLink/Source/LiveLinkMovieScene/Public/MovieScene/MovieSceneLiveLinkBufferData.h
@@ -3,7 +3,6 @@
 #pragma once
 
 #include "Channels/MovieSceneFloatChannel.h"
-#include "Channels/MovieSceneChannelProxy.h"
 
 
 /** Data structure to handle buffering transform keys. Inspired from 3d Transform Track Recorder */
@@ -137,3 +136,7 @@ struct FLiveLinkTransformKeys
 		}
 	}
 };
+
+#if UE_ENABLE_INCLUDE_ORDER_DEPRECATED_IN_5_2
+#include "Channels/MovieSceneChannelProxy.h"
+#endif
diff --git a/Engine/Plugins/Animation/LiveLink/Source/LiveLinkMovieScene/Public/MovieScene/MovieSceneLiveLinkEnumHandler.h b/Engine/Plugins/Animation/LiveLink/Source/LiveLinkMovieScene/Public/MovieScene/MovieSceneLiveLinkEnumHandler.h
--- a/Engine/Plugins/Animation/LiveLink/Source/LiveLinkMovieScene/Public/MovieScene/MovieSceneLiveLinkEnumHandler.h
+++ b/Engine/Plugins/Animation/LiveLink/Source/LiveLinkMovieScene/Public/MovieScene/MovieSceneLiveLinkEnumHandler.h
@@ -2,12 +2,12 @@
 
 #pragma once
 
-#include "CoreMinimal.h"
 
 #include "MovieScene/IMovieSceneLiveLinkPropertyHandler.h"
-#include "MovieScene/MovieSceneLiveLinkStructProperties.h"
 #include "MovieScene/MovieSceneLiveLinkStructPropertyBindings.h"
 
+struct FLiveLinkPropertyData;
+
 
 struct FKeyDataOptimizationParams;
 
@@ -57,3 +57,8 @@ protected:
 
 
 
+
+#if UE_ENABLE_INCLUDE_ORDER_DEPRECATED_IN_5_2
+#include "CoreMinimal.h"
+#include "MovieScene/MovieSceneLiveLinkStructProperties.h"
+#endif
diff --git a/Engine/Plugins/Animation/LiveLink/Source/LiveLinkMovieScene/Public/MovieScene/MovieSceneLiveLinkSection.h b/Engine/Plugins/Animation/LiveLink/Source/LiveLinkMovieScene/Public/MovieScene/MovieSceneLiveLinkSection.h
--- a/Engine/Plugins/Animation/LiveLink/Source/LiveLinkMovieScene/Public/MovieScene/MovieSceneLiveLinkSection.h
+++ b/Engine/Plugins/Animation/LiveLink/Source/LiveLinkMovieScene/Public/MovieScene/MovieSceneLiveLinkSection.h
@@ -4,13 +4,13 @@
 
 #include "MovieSceneSection.h"
 
-#include "Channels/MovieSceneFloatChannel.h"
-#include "Evaluation/MovieSceneEvalTemplate.h"
 #include "LiveLinkPresetTypes.h"
-#include "LiveLinkRole.h"
 
 #include "MovieSceneLiveLinkSection.generated.h"
 
+struct FMovieSceneEvalTemplatePtr;
+struct FMovieSceneFloatChannel;
+
 class UMovieScenePropertyTrack;
 struct FKeyDataOptimizationParams;
 struct FLiveLinkSubjectPreset;
@@ -100,3 +100,9 @@ public:
 	TArray <FMovieSceneFloatChannel> PropertyFloatChannels_DEPRECATED;
 };
 PRAGMA_ENABLE_DEPRECATION_WARNINGS
+
+#if UE_ENABLE_INCLUDE_ORDER_DEPRECATED_IN_5_2
+#include "Channels/MovieSceneFloatChannel.h"
+#include "Evaluation/MovieSceneEvalTemplate.h"
+#include "LiveLinkRole.h"
+#endif
diff --git a/Engine/Plugins/Animation/LiveLink/Source/LiveLinkMovieScene/Public/MovieScene/MovieSceneLiveLinkStructProperties.h b/Engine/Plugins/Animation/LiveLink/Source/LiveLinkMovieScene/Public/MovieScene/MovieSceneLiveLinkStructProperties.h
--- a/Engine/Plugins/Animation/LiveLink/Source/LiveLinkMovieScene/Public/MovieScene/MovieSceneLiveLinkStructProperties.h
+++ b/Engine/Plugins/Animation/LiveLink/Source/LiveLinkMovieScene/Public/MovieScene/MovieSceneLiveLinkStructProperties.h
@@ -2,17 +2,15 @@
 
 #pragma once
 
-#include "CoreMinimal.h"
-
-#include "Channels/MovieSceneBoolChannel.h"
-#include "Channels/MovieSceneByteChannel.h"
-#include "Channels/MovieSceneFloatChannel.h"
-#include "Channels/MovieSceneIntegerChannel.h"
-#include "Channels/MovieSceneStringChannel.h"
-#include "LiveLinkTypes.h"
-
 #include "MovieSceneLiveLinkStructProperties.generated.h"
 
+struct FLiveLinkPropertyData;
+struct FMovieSceneBoolChannel;
+struct FMovieSceneByteChannel;
+struct FMovieSceneFloatChannel;
+struct FMovieSceneIntegerChannel;
+struct FMovieSceneStringChannel;
+
 class IMovieSceneLiveLinkPropertyHandler;
 
 namespace LiveLinkPropertiesUtils
@@ -26,6 +24,9 @@ struct FLiveLinkPropertyData
 	GENERATED_BODY()
 
 public:
+	FLiveLinkPropertyData();
+	~FLiveLinkPropertyData();
+
 	UPROPERTY()
 	FName PropertyName;
 
@@ -60,3 +61,13 @@ public:
 	TArray<FLiveLinkPropertyData> Properties;
 };
 
+
+#if UE_ENABLE_INCLUDE_ORDER_DEPRECATED_IN_5_2
+#include "Channels/MovieSceneBoolChannel.h"
+#include "Channels/MovieSceneByteChannel.h"
+#include "Channels/MovieSceneFloatChannel.h"
+#include "Channels/MovieSceneIntegerChannel.h"
+#include "Channels/MovieSceneStringChannel.h"
+#include "CoreMinimal.h"
+#include "LiveLinkTypes.h"
+#endif
diff --git a/Engine/Plugins/Animation/LiveLink/Source/LiveLinkMovieScene/Public/MovieScene/MovieSceneLiveLinkStructPropertyBindings.h b/Engine/Plugins/Animation/LiveLink/Source/LiveLinkMovieScene/Public/MovieScene/MovieSceneLiveLinkStructPropertyBindings.h
--- a/Engine/Plugins/Animation/LiveLink/Source/LiveLinkMovieScene/Public/MovieScene/MovieSceneLiveLinkStructPropertyBindings.h
+++ b/Engine/Plugins/Animation/LiveLink/Source/LiveLinkMovieScene/Public/MovieScene/MovieSceneLiveLinkStructPropertyBindings.h
@@ -2,15 +2,9 @@
 
 #pragma once
 
-#include "CoreMinimal.h"
 
-#include "Channels/MovieSceneBoolChannel.h"
-#include "Channels/MovieSceneByteChannel.h"
-#include "Channels/MovieSceneFloatChannel.h"
-#include "Channels/MovieSceneIntegerChannel.h"
-#include "Channels/MovieSceneStringChannel.h"
+#include "UObject/UnrealType.h"
 #include "UObject/WeakFieldPtr.h"
-#include "LiveLinkTypes.h"
 
 
 /**
@@ -293,3 +287,13 @@ LIVELINKMOVIESCENE_API void FLiveLinkStructPropertyBindings::SetCurrentValue<boo
 
 template<> LIVELINKMOVIESCENE_API bool FLiveLinkStructPropertyBindings::GetCurrentValueAt<bool>(int32 InIndex, const UScriptStruct& InStruct, const void* InSourceAddress);
 template<> LIVELINKMOVIESCENE_API void FLiveLinkStructPropertyBindings::SetCurrentValueAt<bool>(int32 InIndex, const UScriptStruct& InStruct, void* InSourceAddress, TCallTraits<bool>::ParamType InValue);
+
+#if UE_ENABLE_INCLUDE_ORDER_DEPRECATED_IN_5_2
+#include "Channels/MovieSceneBoolChannel.h"
+#include "Channels/MovieSceneByteChannel.h"
+#include "Channels/MovieSceneFloatChannel.h"
+#include "Channels/MovieSceneIntegerChannel.h"
+#include "Channels/MovieSceneStringChannel.h"
+#include "CoreMinimal.h"
+#include "LiveLinkTypes.h"
+#endif
diff --git a/Engine/Plugins/Animation/LiveLink/Source/LiveLinkMovieScene/Public/MovieScene/MovieSceneLiveLinkSubSection.h b/Engine/Plugins/Animation/LiveLink/Source/LiveLinkMovieScene/Public/MovieScene/MovieSceneLiveLinkSubSection.h
--- a/Engine/Plugins/Animation/LiveLink/Source/LiveLinkMovieScene/Public/MovieScene/MovieSceneLiveLinkSubSection.h
+++ b/Engine/Plugins/Animation/LiveLink/Source/LiveLinkMovieScene/Public/MovieScene/MovieSceneLiveLinkSubSection.h
@@ -2,18 +2,16 @@
 
 #pragma once
 
-#include "MovieSceneSection.h"
 
-#include "Channels/MovieSceneChannelProxy.h"
-#include "Channels/MovieSceneFloatChannel.h"
-#include "Evaluation/MovieSceneEvalTemplate.h"
-#include "LiveLinkRole.h"
+#include "Channels/MovieSceneChannelProxy.h" // IWYU pragma: keep
 #include "LiveLinkTypes.h"
 #include "MovieSceneLiveLinkStructProperties.h"
 #include "Templates/SubclassOf.h"
-
 #include "MovieSceneLiveLinkSubSection.generated.h"
 
+class ULiveLinkRole;
+struct FMovieSceneChannelProxyData;
+
 class UMovieScenePropertyTrack;
 struct FKeyDataOptimizationParams;
 
@@ -105,3 +103,10 @@ protected:
 	TSharedPtr<FLiveLinkStaticDataStruct> StaticData;
 };
 
+
+#if UE_ENABLE_INCLUDE_ORDER_DEPRECATED_IN_5_2
+#include "Channels/MovieSceneFloatChannel.h"
+#include "Evaluation/MovieSceneEvalTemplate.h"
+#include "LiveLinkRole.h"
+#include "MovieSceneSection.h"
+#endif
diff --git a/Engine/Plugins/Animation/LiveLink/Source/LiveLinkMovieScene/Public/MovieScene/MovieSceneLiveLinkSubSectionAnimation.h b/Engine/Plugins/Animation/LiveLink/Source/LiveLinkMovieScene/Public/MovieScene/MovieSceneLiveLinkSubSectionAnimation.h
--- a/Engine/Plugins/Animation/LiveLink/Source/LiveLinkMovieScene/Public/MovieScene/MovieSceneLiveLinkSubSectionAnimation.h
+++ b/Engine/Plugins/Animation/LiveLink/Source/LiveLinkMovieScene/Public/MovieScene/MovieSceneLiveLinkSubSectionAnimation.h
@@ -4,14 +4,11 @@
 
 #include "MovieSceneLiveLinkSubSection.h"
 
-#include "Channels/MovieSceneFloatChannel.h"
-#include "Containers/Array.h"
-#include "LiveLinkRole.h"
-#include "LiveLinkTypes.h"
-#include "MovieScene/MovieSceneLiveLinkTransformHandler.h"
 
 #include "MovieSceneLiveLinkSubSectionAnimation.generated.h"
 
+class FMovieSceneLiveLinkTransformHandler;
+
 
 /**
  * A LiveLinkSubSection managing array of transforms contained in the Animation Frame Data structure
@@ -42,3 +39,9 @@ protected:
 	/** Helper struct to manage filling channels from transforms */
 	TSharedPtr<FMovieSceneLiveLinkTransformHandler> TransformHandler;
 };
+
+#if UE_ENABLE_INCLUDE_ORDER_DEPRECATED_IN_5_2
+#include "Channels/MovieSceneFloatChannel.h"
+#include "LiveLinkRole.h"
+#include "MovieScene/MovieSceneLiveLinkTransformHandler.h"
+#endif
diff --git a/Engine/Plugins/Animation/LiveLink/Source/LiveLinkMovieScene/Public/MovieScene/MovieSceneLiveLinkSubSectionBasicRole.h b/Engine/Plugins/Animation/LiveLink/Source/LiveLinkMovieScene/Public/MovieScene/MovieSceneLiveLinkSubSectionBasicRole.h
--- a/Engine/Plugins/Animation/LiveLink/Source/LiveLinkMovieScene/Public/MovieScene/MovieSceneLiveLinkSubSectionBasicRole.h
+++ b/Engine/Plugins/Animation/LiveLink/Source/LiveLinkMovieScene/Public/MovieScene/MovieSceneLiveLinkSubSectionBasicRole.h
@@ -4,11 +4,6 @@
 
 #include "MovieSceneLiveLinkSubSection.h"
 
-#include "Channels/MovieSceneFloatChannel.h"
-#include "Containers/Array.h"
-#include "LiveLinkRole.h"
-#include "LiveLinkTypes.h"
-#include "MovieScene/MovieSceneLiveLinkPropertyHandler.h"
 
 #include "MovieSceneLiveLinkSubSectionBasicRole.generated.h"
 
@@ -42,3 +37,9 @@ protected:
 	/** Helper struct to manage filling channels from property array */
 	TSharedPtr<IMovieSceneLiveLinkPropertyHandler> PropertyHandler;
 };
+
+#if UE_ENABLE_INCLUDE_ORDER_DEPRECATED_IN_5_2
+#include "Channels/MovieSceneFloatChannel.h"
+#include "LiveLinkRole.h"
+#include "MovieScene/MovieSceneLiveLinkPropertyHandler.h"
+#endif
diff --git a/Engine/Plugins/Animation/LiveLink/Source/LiveLinkMovieScene/Public/MovieScene/MovieSceneLiveLinkSubSectionProperties.h b/Engine/Plugins/Animation/LiveLink/Source/LiveLinkMovieScene/Public/MovieScene/MovieSceneLiveLinkSubSectionProperties.h
--- a/Engine/Plugins/Animation/LiveLink/Source/LiveLinkMovieScene/Public/MovieScene/MovieSceneLiveLinkSubSectionProperties.h
+++ b/Engine/Plugins/Animation/LiveLink/Source/LiveLinkMovieScene/Public/MovieScene/MovieSceneLiveLinkSubSectionProperties.h
@@ -4,9 +4,6 @@
 
 #include "MovieSceneLiveLinkSubSection.h"
 
-#include "LiveLinkRole.h"
-#include "LiveLinkTypes.h"
-#include "MovieScene/MovieSceneLiveLinkPropertyHandler.h"
 
 #include "MovieSceneLiveLinkSubSectionProperties.generated.h"
 
@@ -46,3 +43,8 @@ protected:
 	TArray<TSharedPtr<IMovieSceneLiveLinkPropertyHandler>> PropertyHandlers;
 };
 
+
+#if UE_ENABLE_INCLUDE_ORDER_DEPRECATED_IN_5_2
+#include "LiveLinkRole.h"
+#include "MovieScene/MovieSceneLiveLinkPropertyHandler.h"
+#endif
diff --git a/Engine/Plugins/Animation/LiveLink/Source/LiveLinkMovieScene/Public/MovieScene/MovieSceneLiveLinkTransformHandler.h b/Engine/Plugins/Animation/LiveLink/Source/LiveLinkMovieScene/Public/MovieScene/MovieSceneLiveLinkTransformHandler.h
--- a/Engine/Plugins/Animation/LiveLink/Source/LiveLinkMovieScene/Public/MovieScene/MovieSceneLiveLinkTransformHandler.h
+++ b/Engine/Plugins/Animation/LiveLink/Source/LiveLinkMovieScene/Public/MovieScene/MovieSceneLiveLinkTransformHandler.h
@@ -2,13 +2,13 @@
 
 #pragma once
 
-#include "CoreMinimal.h"
 
 #include "MovieScene/IMovieSceneLiveLinkPropertyHandler.h"
-#include "MovieScene/MovieSceneLiveLinkStructProperties.h"
 #include "MovieScene/MovieSceneLiveLinkStructPropertyBindings.h"
 #include "MovieScene/MovieSceneLiveLinkBufferData.h"
 
+struct FLiveLinkPropertyData;
+
 
 struct FKeyDataOptimizationParams;
 
@@ -52,3 +52,8 @@ protected:
 
 
 
+
+#if UE_ENABLE_INCLUDE_ORDER_DEPRECATED_IN_5_2
+#include "CoreMinimal.h"
+#include "MovieScene/MovieSceneLiveLinkStructProperties.h"
+#endif
diff --git a/Engine/Plugins/Animation/LiveLink/Source/LiveLinkSequencer/Public/LiveLinkSequencerSettings.h b/Engine/Plugins/Animation/LiveLink/Source/LiveLinkSequencer/Public/LiveLinkSequencerSettings.h
--- a/Engine/Plugins/Animation/LiveLink/Source/LiveLinkSequencer/Public/LiveLinkSequencerSettings.h
+++ b/Engine/Plugins/Animation/LiveLink/Source/LiveLinkSequencer/Public/LiveLinkSequencerSettings.h
@@ -4,12 +4,13 @@
 
 #include "Engine/DeveloperSettings.h"
 
-#include "LiveLinkControllerBase.h"
-#include "MovieSceneLiveLinkControllerTrackRecorder.h"
 #include "Templates/SubclassOf.h"
 
 #include "LiveLinkSequencerSettings.generated.h"
 
+class ULiveLinkControllerBase;
+class UMovieSceneLiveLinkControllerTrackRecorder;
+
 
 /**
  * Settings for LiveLink Sequencer
@@ -34,3 +35,8 @@ public:
 	UPROPERTY(config, EditAnywhere, Category = "LiveLink", meta = (AllowAbstract = "false"))
 	TMap<TSubclassOf<ULiveLinkControllerBase>, TSubclassOf<UMovieSceneLiveLinkControllerTrackRecorder>> DefaultTrackRecordersForController;
 };
+
+#if UE_ENABLE_INCLUDE_ORDER_DEPRECATED_IN_5_2
+#include "LiveLinkControllerBase.h"
+#include "MovieSceneLiveLinkControllerTrackRecorder.h"
+#endif
diff --git a/Engine/Plugins/Animation/LiveLink/Source/LiveLinkSequencer/Public/MovieSceneLiveLinkControllerMapTrackRecorder.h b/Engine/Plugins/Animation/LiveLink/Source/LiveLinkSequencer/Public/MovieSceneLiveLinkControllerMapTrackRecorder.h
--- a/Engine/Plugins/Animation/LiveLink/Source/LiveLinkSequencer/Public/MovieSceneLiveLinkControllerMapTrackRecorder.h
+++ b/Engine/Plugins/Animation/LiveLink/Source/LiveLinkSequencer/Public/MovieSceneLiveLinkControllerMapTrackRecorder.h
@@ -5,10 +5,12 @@
 #include "TrackRecorders/IMovieSceneTrackRecorderFactory.h"
 #include "TrackRecorders/MovieSceneTrackRecorder.h"
 
-#include "MovieSceneLiveLinkControllerTrackRecorder.h"
 
 #include "MovieSceneLiveLinkControllerMapTrackRecorder.generated.h"
 
+class ULiveLinkControllerBase;
+class UMovieSceneLiveLinkControllerTrackRecorder;
+
 /** Movie Scene track recorder factory for the LiveLink Component's Controller Map */
 class LIVELINKSEQUENCER_API FMovieSceneLiveLinkControllerMapTrackRecorderFactory : public IMovieSceneTrackRecorderFactory
 {
@@ -51,3 +53,7 @@ private:
 	UPROPERTY(Transient)
 	TArray<TObjectPtr<UMovieSceneLiveLinkControllerTrackRecorder>> ControllerRecorders;
 };
+
+#if UE_ENABLE_INCLUDE_ORDER_DEPRECATED_IN_5_2
+#include "MovieSceneLiveLinkControllerTrackRecorder.h"
+#endif
diff --git a/Engine/Plugins/Animation/LiveLink/Source/LiveLinkSequencer/Public/MovieSceneLiveLinkTrackRecorder.h b/Engine/Plugins/Animation/LiveLink/Source/LiveLinkSequencer/Public/MovieSceneLiveLinkTrackRecorder.h
--- a/Engine/Plugins/Animation/LiveLink/Source/LiveLinkSequencer/Public/MovieSceneLiveLinkTrackRecorder.h
+++ b/Engine/Plugins/Animation/LiveLink/Source/LiveLinkSequencer/Public/MovieSceneLiveLinkTrackRecorder.h
@@ -4,17 +4,12 @@
 
 #include "TrackRecorders/MovieSceneTrackRecorder.h"
 
-#include "CoreMinimal.h"
-#include "GameFramework/Actor.h"
-#include "ILiveLinkClient.h"
-#include "LiveLinkTypes.h"
-#include "Misc/Guid.h"
-#include "MovieScene.h"
-#include "TrackRecorders/IMovieSceneTrackRecorderFactory.h"
 #include "MovieScene/MovieSceneLiveLinkSection.h"
 
 #include "MovieSceneLiveLinkTrackRecorder.generated.h"
 
+class UMovieSceneTrackRecorderSettings;
+
 
 class UMovieSceneLiveLinkTrack;
 class UMovieSceneLiveLinkSectionBase;
@@ -107,3 +102,10 @@ private:
 
 	TArray<FLiveLinkFrameDataStruct> FramesToProcess;
 };
+
+#if UE_ENABLE_INCLUDE_ORDER_DEPRECATED_IN_5_2
+#include "CoreMinimal.h"
+#include "GameFramework/Actor.h"
+#include "ILiveLinkClient.h"
+#include "TrackRecorders/IMovieSceneTrackRecorderFactory.h"
+#endif
diff --git a/Engine/Plugins/Animation/MotionWarping/Source/MotionWarping/Public/AnimNotifyState_MotionWarping.h b/Engine/Plugins/Animation/MotionWarping/Source/MotionWarping/Public/AnimNotifyState_MotionWarping.h
--- a/Engine/Plugins/Animation/MotionWarping/Source/MotionWarping/Public/AnimNotifyState_MotionWarping.h
+++ b/Engine/Plugins/Animation/MotionWarping/Source/MotionWarping/Public/AnimNotifyState_MotionWarping.h
@@ -2,7 +2,6 @@
 
 #pragma once
 
-#include "CoreMinimal.h"
 #include "Animation/AnimNotifies/AnimNotifyState.h"
 #include "AnimNotifyState_MotionWarping.generated.h"
 
@@ -53,4 +52,8 @@ public:
 #if WITH_EDITOR
 	virtual void ValidateAssociatedAssets() override;
 #endif
-};
\ No newline at end of file
+};
+
+#if UE_ENABLE_INCLUDE_ORDER_DEPRECATED_IN_5_2
+#include "CoreMinimal.h"
+#endif
diff --git a/Engine/Plugins/Animation/MotionWarping/Source/MotionWarping/Public/MotionWarping.h b/Engine/Plugins/Animation/MotionWarping/Source/MotionWarping/Public/MotionWarping.h
--- a/Engine/Plugins/Animation/MotionWarping/Source/MotionWarping/Public/MotionWarping.h
+++ b/Engine/Plugins/Animation/MotionWarping/Source/MotionWarping/Public/MotionWarping.h
@@ -2,8 +2,7 @@
 
 #pragma once
 
-#include "CoreMinimal.h"
-#include "Modules/ModuleManager.h"
+#include "Modules/ModuleInterface.h"
 
 class FMotionWarpingModule : public IModuleInterface
 {
@@ -13,3 +12,8 @@ public:
 	virtual void StartupModule() override;
 	virtual void ShutdownModule() override;
 };
+
+#if UE_ENABLE_INCLUDE_ORDER_DEPRECATED_IN_5_2
+#include "CoreMinimal.h"
+#include "Modules/ModuleManager.h"
+#endif
diff --git a/Engine/Plugins/Animation/MotionWarping/Source/MotionWarping/Public/MotionWarpingComponent.h b/Engine/Plugins/Animation/MotionWarping/Source/MotionWarping/Public/MotionWarpingComponent.h
--- a/Engine/Plugins/Animation/MotionWarping/Source/MotionWarping/Public/MotionWarpingComponent.h
+++ b/Engine/Plugins/Animation/MotionWarping/Source/MotionWarping/Public/MotionWarpingComponent.h
@@ -2,13 +2,12 @@
 
 #pragma once
 
-#include "CoreMinimal.h"
 #include "Components/ActorComponent.h"
-#include "HAL/IConsoleManager.h"
-#include "Kismet/KismetSystemLibrary.h"
 #include "RootMotionModifier.h"
 #include "MotionWarpingComponent.generated.h"
 
+template <class T> class TAutoConsoleVariable;
+
 class ACharacter;
 class UAnimSequenceBase;
 class UCharacterMovementComponent;
@@ -196,4 +195,10 @@ protected:
 	bool FindAndUpdateWarpTarget(const FMotionWarpingTarget& WarpTarget);
 
 	FTransform ProcessRootMotionPreConvertToWorld(const FTransform& InRootMotion, class UCharacterMovementComponent* CharacterMovementComponent, float DeltaSeconds);
-};
\ No newline at end of file
+};
+
+#if UE_ENABLE_INCLUDE_ORDER_DEPRECATED_IN_5_2
+#include "CoreMinimal.h"
+#include "HAL/IConsoleManager.h"
+#include "Kismet/KismetSystemLibrary.h"
+#endif
diff --git a/Engine/Plugins/Animation/MotionWarping/Source/MotionWarping/Public/RootMotionModifier.h b/Engine/Plugins/Animation/MotionWarping/Source/MotionWarping/Public/RootMotionModifier.h
--- a/Engine/Plugins/Animation/MotionWarping/Source/MotionWarping/Public/RootMotionModifier.h
+++ b/Engine/Plugins/Animation/MotionWarping/Source/MotionWarping/Public/RootMotionModifier.h
@@ -2,9 +2,8 @@
 
 #pragma once
 
-#include "CoreMinimal.h"
-#include "Animation/AnimSequence.h"
 #include "AlphaBlend.h"
+#include "Animation/AnimSequenceBase.h"
 #include "RootMotionModifier.generated.h"
 
 class UMotionWarpingComponent;
@@ -355,4 +354,9 @@ public:
 		UPARAM(DisplayName = "Start Time") float InStartTime,
 		UPARAM(DisplayName = "End Time") float InEndTime,
 		UPARAM(DisplayName = "Scale") FVector InScale = FVector(1.f));
-};
\ No newline at end of file
+};
+
+#if UE_ENABLE_INCLUDE_ORDER_DEPRECATED_IN_5_2
+#include "Animation/AnimSequence.h"
+#include "CoreMinimal.h"
+#endif
diff --git a/Engine/Plugins/Animation/MotionWarping/Source/MotionWarping/Public/RootMotionModifier_AdjustmentBlendWarp.h b/Engine/Plugins/Animation/MotionWarping/Source/MotionWarping/Public/RootMotionModifier_AdjustmentBlendWarp.h
--- a/Engine/Plugins/Animation/MotionWarping/Source/MotionWarping/Public/RootMotionModifier_AdjustmentBlendWarp.h
+++ b/Engine/Plugins/Animation/MotionWarping/Source/MotionWarping/Public/RootMotionModifier_AdjustmentBlendWarp.h
@@ -2,11 +2,13 @@
 
 #pragma once
 
-#include "BonePose.h"
-#include "CoreMinimal.h"
+#include "Animation/AnimSequence.h"
 #include "RootMotionModifier.h"
 #include "RootMotionModifier_AdjustmentBlendWarp.generated.h"
 
+class ACharacter;
+template <class PoseType> struct FCSPose;
+
 USTRUCT()
 struct FMotionDeltaTrack
 {
@@ -100,4 +102,9 @@ protected:
 	static void ExtractMotionDeltaFromRange(const FBoneContainer& BoneContainer, const UAnimSequenceBase* Animation, float StartTime, float EndTime, float SampleRate, FMotionDeltaTrackContainer& OutMotionDeltaTracks);
 
 	static void AdjustmentBlendWarp(const FBoneContainer& BoneContainer, const FCSPose<FCompactPose>& AdditivePose, const FMotionDeltaTrackContainer& MotionDeltaTracks, FAnimSequenceTrackContainer& Output);
-};
\ No newline at end of file
+};
+
+#if UE_ENABLE_INCLUDE_ORDER_DEPRECATED_IN_5_2
+#include "BonePose.h"
+#include "CoreMinimal.h"
+#endif
diff --git a/Engine/Plugins/Bridge/Source/Bridge/Public/IBridgeModule.h b/Engine/Plugins/Bridge/Source/Bridge/Public/IBridgeModule.h
--- a/Engine/Plugins/Bridge/Source/Bridge/Public/IBridgeModule.h
+++ b/Engine/Plugins/Bridge/Source/Bridge/Public/IBridgeModule.h
@@ -1,7 +1,6 @@
 // Copyright Epic Games, Inc. All Rights Reserved.
 #pragma once
 
-#include "Modules/ModuleInterface.h"
 #include "Modules/ModuleManager.h"	
 
 #define MS_MODULE_NAME TEXT("Bridge")
diff --git a/Engine/Plugins/Bridge/Source/MegascansPlugin/Internal/BridgeDragDropHelper.h b/Engine/Plugins/Bridge/Source/MegascansPlugin/Internal/BridgeDragDropHelper.h
--- a/Engine/Plugins/Bridge/Source/MegascansPlugin/Internal/BridgeDragDropHelper.h
+++ b/Engine/Plugins/Bridge/Source/MegascansPlugin/Internal/BridgeDragDropHelper.h
@@ -1,8 +1,10 @@
 // Copyright Epic Games, Inc. All Rights Reserved.
 #pragma once
 
-#include "CoreMinimal.h"
-#include "AssetRegistry/AssetData.h"
+#include "Delegates/Delegate.h"
+
+class AActor;
+struct FAssetData;
 
 class AStaticMeshActor;
 
@@ -23,3 +25,8 @@ public:
 	static void Initialize();
 	static TSharedPtr<FBridgeDragDropHelperImpl> Instance;
 };
+
+#if UE_ENABLE_INCLUDE_ORDER_DEPRECATED_IN_5_2
+#include "AssetRegistry/AssetData.h"
+#include "CoreMinimal.h"
+#endif
diff --git a/Engine/Plugins/Bridge/Source/MegascansPlugin/Public/AssetsImportController.h b/Engine/Plugins/Bridge/Source/MegascansPlugin/Public/AssetsImportController.h
--- a/Engine/Plugins/Bridge/Source/MegascansPlugin/Public/AssetsImportController.h
+++ b/Engine/Plugins/Bridge/Source/MegascansPlugin/Public/AssetsImportController.h
@@ -1,10 +1,7 @@
 // Copyright Epic Games, Inc. All Rights Reserved.
 #pragma once
-#include "CoreMinimal.h"
-#include "Engine/StaticMeshActor.h"
 
-#include "MSAssetImportData.h"
-#include "Materials/MaterialInstanceConstant.h"
+#include "Templates/SharedPointer.h"
 
 class MEGASCANSPLUGIN_API FAssetsImportController
 {
@@ -21,4 +18,11 @@ private:
 public:	
 	static TSharedPtr<FAssetsImportController> Get();
 	void DataReceived(const FString DataFromBridge);
-};
\ No newline at end of file
+};
+
+#if UE_ENABLE_INCLUDE_ORDER_DEPRECATED_IN_5_2
+#include "CoreMinimal.h"
+#include "Engine/StaticMeshActor.h"
+#include "MSAssetImportData.h"
+#include "Materials/MaterialInstanceConstant.h"
+#endif
diff --git a/Engine/Plugins/Bridge/Source/MegascansPlugin/Public/BridgeDragDropUtils.h b/Engine/Plugins/Bridge/Source/MegascansPlugin/Public/BridgeDragDropUtils.h
--- a/Engine/Plugins/Bridge/Source/MegascansPlugin/Public/BridgeDragDropUtils.h
+++ b/Engine/Plugins/Bridge/Source/MegascansPlugin/Public/BridgeDragDropUtils.h
@@ -1,8 +1,9 @@
 // Copyright Epic Games, Inc. All Rights Reserved.
 #pragma once
 
-#include "CoreMinimal.h"
-#include "AssetRegistry/AssetData.h"
+#include "Delegates/Delegate.h"
+
+struct FAssetData;
 
 class AStaticMeshActor;
 
@@ -22,3 +23,8 @@ public:
     static void Initialize();
     static TSharedPtr<FBridgeDragDropImpl> Instance;
 };
+
+#if UE_ENABLE_INCLUDE_ORDER_DEPRECATED_IN_5_2
+#include "AssetRegistry/AssetData.h"
+#include "CoreMinimal.h"
+#endif
diff --git a/Engine/Plugins/Bridge/Source/MegascansPlugin/Public/IMegascansLiveLinkModule.h b/Engine/Plugins/Bridge/Source/MegascansPlugin/Public/IMegascansLiveLinkModule.h
--- a/Engine/Plugins/Bridge/Source/MegascansPlugin/Public/IMegascansLiveLinkModule.h
+++ b/Engine/Plugins/Bridge/Source/MegascansPlugin/Public/IMegascansLiveLinkModule.h
@@ -1,7 +1,6 @@
 // Copyright Epic Games, Inc. All Rights Reserved.
 #pragma once
 
-#include "Modules/ModuleInterface.h"
 #include "Modules/ModuleManager.h"
 
 
diff --git a/Engine/Plugins/Bridge/Source/MegascansPlugin/Public/MSSettings.h b/Engine/Plugins/Bridge/Source/MegascansPlugin/Public/MSSettings.h
--- a/Engine/Plugins/Bridge/Source/MegascansPlugin/Public/MSSettings.h
+++ b/Engine/Plugins/Bridge/Source/MegascansPlugin/Public/MSSettings.h
@@ -1,14 +1,10 @@
 // Copyright Epic Games, Inc. All Rights Reserved.
 #pragma once
 
-#include "CoreMinimal.h"
-#include "Materials/Material.h"
-#include "UObject/ObjectMacros.h"
-#include "UObject/Object.h"
-#include "Engine/EngineTypes.h"
-#include "UObject/SoftObjectPath.h"
 #include "MSSettings.generated.h"
 
+struct FPropertyChangedEvent;
+
 
 UCLASS(Config = Editor)
 class UMegascansSettings
@@ -103,3 +99,9 @@ public:
 
 };
 
+
+#if UE_ENABLE_INCLUDE_ORDER_DEPRECATED_IN_5_2
+#include "CoreMinimal.h"
+#include "Engine/EngineTypes.h"
+#include "Materials/Material.h"
+#endif
diff --git a/Engine/Plugins/Bridge/Source/MegascansPlugin/Public/SMSWindow.h b/Engine/Plugins/Bridge/Source/MegascansPlugin/Public/SMSWindow.h
--- a/Engine/Plugins/Bridge/Source/MegascansPlugin/Public/SMSWindow.h
+++ b/Engine/Plugins/Bridge/Source/MegascansPlugin/Public/SMSWindow.h
@@ -2,10 +2,11 @@
 #pragma once
 
 
-#include "CoreMinimal.h"
-#include "MSSettings.h"
 
-#include "Widgets/SWindow.h"
+#include "Templates/SharedPointer.h"
+
+class SWindow;
+class UMegascansSettings;
 
 class FTabManager;
 
@@ -23,3 +24,9 @@ public:
 
 
 
+
+#if UE_ENABLE_INCLUDE_ORDER_DEPRECATED_IN_5_2
+#include "CoreMinimal.h"
+#include "MSSettings.h"
+#include "Widgets/SWindow.h"
+#endif
diff --git a/Engine/Plugins/Bridge/Source/MegascansPlugin/Public/TCPServer.h b/Engine/Plugins/Bridge/Source/MegascansPlugin/Public/TCPServer.h
--- a/Engine/Plugins/Bridge/Source/MegascansPlugin/Public/TCPServer.h
+++ b/Engine/Plugins/Bridge/Source/MegascansPlugin/Public/TCPServer.h
@@ -2,16 +2,15 @@
 #pragma once
 
 
-#include "CoreMinimal.h"
+#include "Containers/UnrealString.h"
 #include <string>
-#include "Sockets.h"
-#include "Misc/DateTime.h"
-#include "Async/AsyncWork.h"
-#include "Sockets.h"
-#include "Common/TcpListener.h"
-#include "Sockets.h"
 #include "Containers/Queue.h"
-#include "Common/UdpSocketReceiver.h"
+#include "HAL/Runnable.h"
+
+class FRunnableThread;
+class FSocket;
+struct FIPv4Endpoint;
+struct FScriptContainerElement;
 
 class MEGASCANSPLUGIN_API FTCPServer : public FRunnable
 {
@@ -49,3 +48,12 @@ private:
 };
 
 
+
+#if UE_ENABLE_INCLUDE_ORDER_DEPRECATED_IN_5_2
+#include "Async/AsyncWork.h"
+#include "Common/TcpListener.h"
+#include "Common/UdpSocketReceiver.h"
+#include "CoreMinimal.h"
+#include "Misc/DateTime.h"
+#include "Sockets.h"
+#endif
diff --git a/Engine/Plugins/Developer/Concert/ConcertSync/ConcertSyncClient/Source/ConcertSyncClient/Private/ConcertClientLiveTransactionAuthors.cpp b/Engine/Plugins/Developer/Concert/ConcertSync/ConcertSyncClient/Source/ConcertSyncClient/Private/ConcertClientLiveTransactionAuthors.cpp
--- a/Engine/Plugins/Developer/Concert/ConcertSync/ConcertSyncClient/Source/ConcertSyncClient/Private/ConcertClientLiveTransactionAuthors.cpp
+++ b/Engine/Plugins/Developer/Concert/ConcertSync/ConcertSyncClient/Source/ConcertSyncClient/Private/ConcertClientLiveTransactionAuthors.cpp
@@ -3,6 +3,7 @@
 #include "ConcertClientLiveTransactionAuthors.h"
 #include "ConcertSyncClientLiveSession.h"
 #include "ConcertSyncSessionDatabase.h"
+#include "ConcertSyncSessionTypes.h"
 #include "IConcertSession.h"
 
 FConcertClientLiveTransactionAuthors::FConcertClientLiveTransactionAuthors(TSharedRef<FConcertSyncClientLiveSession> InLiveSession)
diff --git a/Engine/Plugins/Developer/Concert/ConcertSync/ConcertSyncClient/Source/ConcertSyncClient/Private/ConcertClientTransactionManager.cpp b/Engine/Plugins/Developer/Concert/ConcertSync/ConcertSyncClient/Source/ConcertSyncClient/Private/ConcertClientTransactionManager.cpp
--- a/Engine/Plugins/Developer/Concert/ConcertSync/ConcertSyncClient/Source/ConcertSyncClient/Private/ConcertClientTransactionManager.cpp
+++ b/Engine/Plugins/Developer/Concert/ConcertSync/ConcertSyncClient/Source/ConcertSyncClient/Private/ConcertClientTransactionManager.cpp
@@ -3,6 +3,7 @@
 #include "ConcertClientTransactionManager.h"
 #include "Algo/AnyOf.h"
 #include "Components/SceneComponent.h"
+#include "ConcertSyncSessionTypes.h"
 #include "ConcertTransactionEvents.h"
 #include "GameFramework/Actor.h"
 #include "IConcertSession.h"
diff --git a/Engine/Plugins/Developer/Concert/ConcertSync/ConcertSyncCore/Source/ConcertSyncCore/Private/HistoryEdition/ActivityDependencyGraph.cpp b/Engine/Plugins/Developer/Concert/ConcertSync/ConcertSyncCore/Source/ConcertSyncCore/Private/HistoryEdition/ActivityDependencyGraph.cpp
--- a/Engine/Plugins/Developer/Concert/ConcertSync/ConcertSyncCore/Source/ConcertSyncCore/Private/HistoryEdition/ActivityDependencyGraph.cpp
+++ b/Engine/Plugins/Developer/Concert/ConcertSync/ConcertSyncCore/Source/ConcertSyncCore/Private/HistoryEdition/ActivityDependencyGraph.cpp
@@ -5,11 +5,9 @@
 #include "ConcertLogGlobal.h"
 #include "CoreGlobals.h"
 #include "HistoryEdition/ActivityDependencyEdge.h"
-#include "HistoryEdition/ActivityGraphIDs.h"
 #include "HistoryEdition/ActivityNode.h"
 
 #include "Algo/ForEach.h"
-#include "Misc/Optional.h"
 
 UE::ConcertSyncCore::FActivityNodeID UE::ConcertSyncCore::FActivityDependencyGraph::AddActivity(int64 ActivityIndex, EActivityNodeFlags NodeFlags)
 {
diff --git a/Engine/Plugins/Developer/Concert/ConcertSync/ConcertSyncCore/Source/ConcertSyncCore/Private/HistoryEdition/ActivityNode.cpp b/Engine/Plugins/Developer/Concert/ConcertSync/ConcertSyncCore/Source/ConcertSyncCore/Private/HistoryEdition/ActivityNode.cpp
--- a/Engine/Plugins/Developer/Concert/ConcertSync/ConcertSyncCore/Source/ConcertSyncCore/Private/HistoryEdition/ActivityNode.cpp
+++ b/Engine/Plugins/Developer/Concert/ConcertSync/ConcertSyncCore/Source/ConcertSyncCore/Private/HistoryEdition/ActivityNode.cpp
@@ -4,8 +4,6 @@
 
 #include "HistoryEdition/ActivityDependencyEdge.h"
 #include "HistoryEdition/ActivityDependencyGraph.h"
-#include "HistoryEdition/ActivityGraphIDs.h"
-#include "Misc/Optional.h"
 
 bool UE::ConcertSyncCore::FActivityNode::DependsOnNode(FActivityNodeID NodeId, TOptional<EActivityDependencyReason> WithReason, TOptional<EDependencyStrength> WithStrength) const
 {
diff --git a/Engine/Plugins/Developer/Concert/ConcertSync/ConcertSyncCore/Source/ConcertSyncCore/Private/HistoryEdition/DebugDependencyGraph.cpp b/Engine/Plugins/Developer/Concert/ConcertSync/ConcertSyncCore/Source/ConcertSyncCore/Private/HistoryEdition/DebugDependencyGraph.cpp
--- a/Engine/Plugins/Developer/Concert/ConcertSync/ConcertSyncCore/Source/ConcertSyncCore/Private/HistoryEdition/DebugDependencyGraph.cpp
+++ b/Engine/Plugins/Developer/Concert/ConcertSync/ConcertSyncCore/Source/ConcertSyncCore/Private/HistoryEdition/DebugDependencyGraph.cpp
@@ -9,7 +9,6 @@
 #include "ConcertLogGlobal.h"
 #include "ConcertSyncSessionDatabase.h"
 #include "HistoryEdition/ActivityDependencyGraph.h"
-#include "HistoryEdition/ActivityGraphIDs.h"
 
 void UE::ConcertSyncCore::Graphviz::MakeNodeTitle(FGraphStringBuilder& WriteTo, FActivityNodeID ToStringify, const FActivityDependencyGraph& Graph, const FConcertSyncSessionDatabase& SessionDatabase, ENodeTitleFlags NodeTitleFlags)
 {
diff --git a/Engine/Plugins/Developer/Concert/ConcertSync/ConcertSyncCore/Source/ConcertSyncCore/Private/HistoryEdition/DependencyGraphBuilder.cpp b/Engine/Plugins/Developer/Concert/ConcertSync/ConcertSyncCore/Source/ConcertSyncCore/Private/HistoryEdition/DependencyGraphBuilder.cpp
--- a/Engine/Plugins/Developer/Concert/ConcertSync/ConcertSyncCore/Source/ConcertSyncCore/Private/HistoryEdition/DependencyGraphBuilder.cpp
+++ b/Engine/Plugins/Developer/Concert/ConcertSync/ConcertSyncCore/Source/ConcertSyncCore/Private/HistoryEdition/DependencyGraphBuilder.cpp
@@ -6,7 +6,6 @@
 #include "ConcertSyncSessionTypes.h"
 #include "HistoryEdition/ActivityDependencyEdge.h"
 #include "HistoryEdition/ActivityDependencyGraph.h"
-#include "HistoryEdition/ActivityGraphIDs.h"
 #include "HistoryEdition/ActivityNode.h"
 
 namespace UE::ConcertSyncCore
diff --git a/Engine/Plugins/Developer/Concert/ConcertSync/ConcertSyncCore/Source/ConcertSyncCore/Public/ConcertDataStore.h b/Engine/Plugins/Developer/Concert/ConcertSync/ConcertSyncCore/Source/ConcertSyncCore/Public/ConcertDataStore.h
--- a/Engine/Plugins/Developer/Concert/ConcertSync/ConcertSyncCore/Source/ConcertSyncCore/Public/ConcertDataStore.h
+++ b/Engine/Plugins/Developer/Concert/ConcertSync/ConcertSyncCore/Source/ConcertSyncCore/Public/ConcertDataStore.h
@@ -2,9 +2,10 @@
 
 #pragma once
 
-#include "CoreMinimal.h"
 #include "ConcertDataStoreMessages.h"
 
+struct FConcertSessionSerializedPayload;
+
 /**
  * Maintains a type-safe key/value map where the values are USTRUCT() structs
  * serialized into a FConcertDataStore_StoreValue. Each key/value pair has a version,
@@ -131,3 +132,7 @@ private:
 	/** Maps keyName/keyValue. */
 	TMap<FName, FConcertDataStoreValueRef> KeyValueMap;
 };
+
+#if UE_ENABLE_INCLUDE_ORDER_DEPRECATED_IN_5_2
+#include "CoreMinimal.h"
+#endif
diff --git a/Engine/Plugins/Developer/Concert/ConcertSync/ConcertSyncCore/Source/ConcertSyncCore/Public/ConcertDataStoreMessages.h b/Engine/Plugins/Developer/Concert/ConcertSync/ConcertSyncCore/Source/ConcertSyncCore/Public/ConcertDataStoreMessages.h
--- a/Engine/Plugins/Developer/Concert/ConcertSync/ConcertSyncCore/Source/ConcertSyncCore/Public/ConcertDataStoreMessages.h
+++ b/Engine/Plugins/Developer/Concert/ConcertSync/ConcertSyncCore/Source/ConcertSyncCore/Public/ConcertDataStoreMessages.h
@@ -2,9 +2,7 @@
 
 #pragma once
 
-#include "CoreMinimal.h"
-#include "ConcertMessages.h"
-#include "UObject/ObjectMacros.h"
+#include "ConcertMessageData.h"
 #include "UObject/StructOnScope.h"
 #include "ConcertDataStoreMessages.generated.h"
 
@@ -512,3 +510,8 @@ private:
 	/** The weakly-typed operation result. */
 	FConcertDataStoreResult Result;
 };
+
+#if UE_ENABLE_INCLUDE_ORDER_DEPRECATED_IN_5_2
+#include "ConcertMessages.h"
+#include "CoreMinimal.h"
+#endif
diff --git a/Engine/Plugins/Developer/Concert/ConcertSync/ConcertSyncCore/Source/ConcertSyncCore/Public/ConcertFileCache.h b/Engine/Plugins/Developer/Concert/ConcertSync/ConcertSyncCore/Source/ConcertSyncCore/Public/ConcertFileCache.h
--- a/Engine/Plugins/Developer/Concert/ConcertSync/ConcertSyncCore/Source/ConcertSyncCore/Public/ConcertFileCache.h
+++ b/Engine/Plugins/Developer/Concert/ConcertSync/ConcertSyncCore/Source/ConcertSyncCore/Public/ConcertFileCache.h
@@ -2,8 +2,9 @@
 
 #pragma once
 
-#include "CoreMinimal.h"
 #include "Containers/LruCache.h"
+#include "Misc/DateTime.h"
+#include "Templates/SharedPointer.h"
 
 /**
  * Maintains an in-memory cache of file data, attempting to keep the cache within the given upper memory threshold.
@@ -115,3 +116,7 @@ private:
 	uint64 TotalCachedFileDataBytes;
 	TLruCache<FString, TSharedPtr<FInternalCacheEntry>> InternalCache;
 };
+
+#if UE_ENABLE_INCLUDE_ORDER_DEPRECATED_IN_5_2
+#include "CoreMinimal.h"
+#endif
diff --git a/Engine/Plugins/Developer/Concert/ConcertSync/ConcertSyncCore/Source/ConcertSyncCore/Public/ConcertPackageEvents.h b/Engine/Plugins/Developer/Concert/ConcertSync/ConcertSyncCore/Source/ConcertSyncCore/Public/ConcertPackageEvents.h
--- a/Engine/Plugins/Developer/Concert/ConcertSync/ConcertSyncCore/Source/ConcertSyncCore/Public/ConcertPackageEvents.h
+++ b/Engine/Plugins/Developer/Concert/ConcertSync/ConcertSyncCore/Source/ConcertSyncCore/Public/ConcertPackageEvents.h
@@ -2,7 +2,8 @@
 
 #pragma once
 
-#include "CoreMinimal.h"
+#include "Delegates/Delegate.h"
+#include "Misc/Guid.h"
 
 struct FConcertPackageInfo;
 
@@ -65,3 +66,7 @@ namespace UE::ConcertSyncCore::ConcertPackageEvents
 	/** Called when we finished receiving a package from a remote endpoint but decided to reject the data. Called only on the server. RemoteEndpointId is to be interpreted as the sender. */
 	CONCERTSYNCCORE_API FConcertRejectSendPackageDelegate& OnRejectRemoteSendPackage();
 }
+
+#if UE_ENABLE_INCLUDE_ORDER_DEPRECATED_IN_5_2
+#include "CoreMinimal.h"
+#endif
diff --git a/Engine/Plugins/Developer/Concert/ConcertSync/ConcertSyncCore/Source/ConcertSyncCore/Public/ConcertSequencerMessages.h b/Engine/Plugins/Developer/Concert/ConcertSync/ConcertSyncCore/Source/ConcertSyncCore/Public/ConcertSequencerMessages.h
--- a/Engine/Plugins/Developer/Concert/ConcertSync/ConcertSyncCore/Source/ConcertSyncCore/Public/ConcertSequencerMessages.h
+++ b/Engine/Plugins/Developer/Concert/ConcertSync/ConcertSyncCore/Source/ConcertSyncCore/Public/ConcertSequencerMessages.h
@@ -2,11 +2,7 @@
 
 #pragma once
 
-#include "CoreMinimal.h"
-#include "UObject/Object.h"
-#include "UObject/Class.h"
 #include "Misc/QualifiedFrameTime.h"
-#include "Math/Range.h"
 #include "ConcertSequencerMessages.generated.h"
 
 
@@ -131,3 +127,7 @@ struct FConcertSequencerTimeAdjustmentEvent
 	UPROPERTY()
 	FString SequenceObjectPath;
 };
+
+#if UE_ENABLE_INCLUDE_ORDER_DEPRECATED_IN_5_2
+#include "CoreMinimal.h"
+#endif
diff --git a/Engine/Plugins/Developer/Concert/ConcertSync/ConcertSyncCore/Source/ConcertSyncCore/Public/ConcertSyncArchives.h b/Engine/Plugins/Developer/Concert/ConcertSync/ConcertSyncCore/Source/ConcertSyncCore/Public/ConcertSyncArchives.h
--- a/Engine/Plugins/Developer/Concert/ConcertSync/ConcertSyncCore/Source/ConcertSyncCore/Public/ConcertSyncArchives.h
+++ b/Engine/Plugins/Developer/Concert/ConcertSync/ConcertSyncCore/Source/ConcertSyncCore/Public/ConcertSyncArchives.h
@@ -2,9 +2,14 @@
 
 #pragma once
 
-#include "CoreMinimal.h"
+#include "Delegates/Delegate.h"
 #include "IdentifierTable/ConcertTransportArchives.h"
 
+struct FLazyObjectPtr;
+struct FObjectPtr;
+struct FSoftObjectPtr;
+struct FWeakObjectPtr;
+
 struct FConcertSessionVersionInfo;
 
 DECLARE_DELEGATE_OneParam(FConcertSyncRemapObjectPath, FString& /*ObjectPath*/)
@@ -100,3 +105,7 @@ private:
 	FConcertSyncWorldRemapper WorldRemapper;
 	FConcertSyncEncounteredMissingObject EncounteredMissingObjectDelegate;
 };
+
+#if UE_ENABLE_INCLUDE_ORDER_DEPRECATED_IN_5_2
+#include "CoreMinimal.h"
+#endif
diff --git a/Engine/Plugins/Developer/Concert/ConcertSync/ConcertSyncCore/Source/ConcertSyncCore/Public/ConcertSyncSessionDatabase.h b/Engine/Plugins/Developer/Concert/ConcertSync/ConcertSyncCore/Source/ConcertSyncCore/Public/ConcertSyncSessionDatabase.h
--- a/Engine/Plugins/Developer/Concert/ConcertSync/ConcertSyncCore/Source/ConcertSyncCore/Public/ConcertSyncSessionDatabase.h
+++ b/Engine/Plugins/Developer/Concert/ConcertSync/ConcertSyncCore/Source/ConcertSyncCore/Public/ConcertSyncSessionDatabase.h
@@ -2,9 +2,28 @@
 
 #pragma once
 
-#include "CoreMinimal.h"
-#include "ConcertSyncSessionTypes.h"
-#include "Templates/SharedPointerInternals.h"
+#include "Delegates/Delegate.h"
+#include "Templates/PimplPtr.h"
+
+class FStructOnScope;
+enum class EConcertSyncActivityEventType : uint8;
+enum class EConcertSyncActivityFlags : uint8;
+struct FConcertPackageDataStream;
+struct FConcertPackageInfo;
+struct FConcertSessionFilter;
+struct FConcertSyncActivity;
+struct FConcertSyncConnectionActivity;
+struct FConcertSyncConnectionEvent;
+struct FConcertSyncEndpointData;
+struct FConcertSyncEndpointIdAndData;
+struct FConcertSyncLockActivity;
+struct FConcertSyncLockEvent;
+struct FConcertSyncPackageEventData;
+struct FConcertSyncTransactionActivity;
+struct FConcertSyncTransactionEvent;
+struct FConcertTransactionFinalizedEvent;
+struct FScriptContainerElement;
+template <typename OptionalType> struct TOptional;
 
 class FConcertFileCache;
 class FConcertSyncSessionDatabaseStatements;
@@ -1004,3 +1023,9 @@ namespace ConcertSyncSessionDatabaseFilterUtil
 	CONCERTSYNCCORE_API bool PackageEventPassesFilter(const int64 InPackageEventId, const FConcertSessionFilter& InSessionFilter, const FConcertSyncSessionDatabase& InDatabase);
 
 } // namespace ConcertSyncSessionDatabaseFilterUtil
+
+#if UE_ENABLE_INCLUDE_ORDER_DEPRECATED_IN_5_2
+#include "ConcertSyncSessionTypes.h"
+#include "CoreMinimal.h"
+#include "Templates/SharedPointerInternals.h"
+#endif
diff --git a/Engine/Plugins/Developer/Concert/ConcertSync/ConcertSyncCore/Source/ConcertSyncCore/Public/ConcertSyncSessionTypes.h b/Engine/Plugins/Developer/Concert/ConcertSync/ConcertSyncCore/Source/ConcertSyncCore/Public/ConcertSyncSessionTypes.h
--- a/Engine/Plugins/Developer/Concert/ConcertSync/ConcertSyncCore/Source/ConcertSyncCore/Public/ConcertSyncSessionTypes.h
+++ b/Engine/Plugins/Developer/Concert/ConcertSync/ConcertSyncCore/Source/ConcertSyncCore/Public/ConcertSyncSessionTypes.h
@@ -2,8 +2,6 @@
 
 #pragma once
 
-#include "CoreMinimal.h"
-#include "ConcertMessageData.h"
 #include "ConcertWorkspaceData.h"
 #include "ConcertTransactionEvents.h"
 #include "UObject/StructOnScope.h"
@@ -455,4 +453,8 @@ struct FConcertSessionActivity
 	 * @see FConcertActivityStream
 	 */
 	TUniquePtr<FConcertSessionSerializedPayload> EventPayload;
-};
\ No newline at end of file
+};
+
+#if UE_ENABLE_INCLUDE_ORDER_DEPRECATED_IN_5_2
+#include "CoreMinimal.h"
+#endif
diff --git a/Engine/Plugins/Developer/Concert/ConcertSync/ConcertSyncCore/Source/ConcertSyncCore/Public/ConcertSyncSettings.h b/Engine/Plugins/Developer/Concert/ConcertSync/ConcertSyncCore/Source/ConcertSyncCore/Public/ConcertSyncSettings.h
--- a/Engine/Plugins/Developer/Concert/ConcertSync/ConcertSyncCore/Source/ConcertSyncCore/Public/ConcertSyncSettings.h
+++ b/Engine/Plugins/Developer/Concert/ConcertSync/ConcertSyncCore/Source/ConcertSyncCore/Public/ConcertSyncSettings.h
@@ -2,8 +2,7 @@
 
 #pragma once
 
-#include "UObject/Object.h"
-#include "UObject/ObjectMacros.h"
+#include "UObject/FieldPath.h"
 #include "ConcertSyncSettings.generated.h"
 
 USTRUCT()
diff --git a/Engine/Plugins/Developer/Concert/ConcertSync/ConcertSyncCore/Source/ConcertSyncCore/Public/ConcertTransactionEvents.h b/Engine/Plugins/Developer/Concert/ConcertSync/ConcertSyncCore/Source/ConcertSyncCore/Public/ConcertTransactionEvents.h
--- a/Engine/Plugins/Developer/Concert/ConcertSync/ConcertSyncCore/Source/ConcertSyncCore/Public/ConcertTransactionEvents.h
+++ b/Engine/Plugins/Developer/Concert/ConcertSync/ConcertSyncCore/Source/ConcertSyncCore/Public/ConcertTransactionEvents.h
@@ -2,10 +2,6 @@
 
 #pragma once
 
-#include "CoreMinimal.h"
-#include "UObject/Object.h"
-#include "UObject/Class.h"
-#include "UObject/Package.h"
 #include "Misc/TransactionObjectEvent.h"
 #include "IdentifierTable/ConcertIdentifierTableData.h"
 #include "ConcertTransactionEvents.generated.h"
@@ -225,3 +221,8 @@ struct FConcertTransactionRejectedEvent
 	UPROPERTY()
 	FGuid TransactionId;
 };
+
+#if UE_ENABLE_INCLUDE_ORDER_DEPRECATED_IN_5_2
+#include "CoreMinimal.h"
+#include "UObject/Package.h"
+#endif
diff --git a/Engine/Plugins/Developer/Concert/ConcertSync/ConcertSyncCore/Source/ConcertSyncCore/Public/HistoryEdition/ActivityDependencyEdge.h b/Engine/Plugins/Developer/Concert/ConcertSync/ConcertSyncCore/Source/ConcertSyncCore/Public/HistoryEdition/ActivityDependencyEdge.h
--- a/Engine/Plugins/Developer/Concert/ConcertSync/ConcertSyncCore/Source/ConcertSyncCore/Public/HistoryEdition/ActivityDependencyEdge.h
+++ b/Engine/Plugins/Developer/Concert/ConcertSync/ConcertSyncCore/Source/ConcertSyncCore/Public/HistoryEdition/ActivityDependencyEdge.h
@@ -2,7 +2,6 @@
 
 #pragma once
 
-#include "CoreMinimal.h"
 #include "ActivityGraphIDs.h"
 
 namespace UE::ConcertSyncCore
@@ -132,4 +131,8 @@ namespace UE::ConcertSyncCore
 		/** How much the activity is depended on */
 		const EDependencyStrength DependencyStrength;
 	};
-}
\ No newline at end of file
+}
+
+#if UE_ENABLE_INCLUDE_ORDER_DEPRECATED_IN_5_2
+#include "CoreMinimal.h"
+#endif
diff --git a/Engine/Plugins/Developer/Concert/ConcertSync/ConcertSyncCore/Source/ConcertSyncCore/Public/HistoryEdition/ActivityDependencyGraph.h b/Engine/Plugins/Developer/Concert/ConcertSync/ConcertSyncCore/Source/ConcertSyncCore/Public/HistoryEdition/ActivityDependencyGraph.h
--- a/Engine/Plugins/Developer/Concert/ConcertSync/ConcertSyncCore/Source/ConcertSyncCore/Public/HistoryEdition/ActivityDependencyGraph.h
+++ b/Engine/Plugins/Developer/Concert/ConcertSync/ConcertSyncCore/Source/ConcertSyncCore/Public/HistoryEdition/ActivityDependencyGraph.h
@@ -2,12 +2,12 @@
 
 #pragma once
 
-#include "CoreMinimal.h"
-#include "ActivityDependencyEdge.h"
-#include "ActivityGraphIDs.h"
 #include "ActivityNode.h"
-#include "Misc/Optional.h"
-#include "Templates/NonNullPointer.h"
+
+template <typename OptionalType> struct TOptional;
+
+namespace UE::ConcertSyncCore { class FActivityDependencyEdge; }
+namespace UE::ConcertSyncCore { struct FActivityNodeID; }
 
 namespace UE::ConcertSyncCore
 {
@@ -82,3 +82,9 @@ namespace UE::ConcertSyncCore
 	};
 }
 
+
+#if UE_ENABLE_INCLUDE_ORDER_DEPRECATED_IN_5_2
+#include "CoreMinimal.h"
+#include "Misc/Optional.h"
+#include "Templates/NonNullPointer.h"
+#endif
diff --git a/Engine/Plugins/Developer/Concert/ConcertSync/ConcertSyncCore/Source/ConcertSyncCore/Public/HistoryEdition/ActivityGraphIDs.h b/Engine/Plugins/Developer/Concert/ConcertSync/ConcertSyncCore/Source/ConcertSyncCore/Public/HistoryEdition/ActivityGraphIDs.h
--- a/Engine/Plugins/Developer/Concert/ConcertSync/ConcertSyncCore/Source/ConcertSyncCore/Public/HistoryEdition/ActivityGraphIDs.h
+++ b/Engine/Plugins/Developer/Concert/ConcertSync/ConcertSyncCore/Source/ConcertSyncCore/Public/HistoryEdition/ActivityGraphIDs.h
@@ -2,8 +2,9 @@
 
 #pragma once
 
-#include "CoreMinimal.h"
+#include "HAL/PlatformMath.h"
 #include "Templates/TypeHash.h"
+#include "Templates/UnrealTemplate.h"
 
 namespace UE::ConcertSyncCore
 {
@@ -46,4 +47,8 @@ namespace UE::ConcertSyncCore
 #endif
 		}
 	};
-}
\ No newline at end of file
+}
+
+#if UE_ENABLE_INCLUDE_ORDER_DEPRECATED_IN_5_2
+#include "CoreMinimal.h"
+#endif
diff --git a/Engine/Plugins/Developer/Concert/ConcertSync/ConcertSyncCore/Source/ConcertSyncCore/Public/HistoryEdition/ActivityNode.h b/Engine/Plugins/Developer/Concert/ConcertSync/ConcertSyncCore/Source/ConcertSyncCore/Public/HistoryEdition/ActivityNode.h
--- a/Engine/Plugins/Developer/Concert/ConcertSync/ConcertSyncCore/Source/ConcertSyncCore/Public/HistoryEdition/ActivityNode.h
+++ b/Engine/Plugins/Developer/Concert/ConcertSync/ConcertSyncCore/Source/ConcertSyncCore/Public/HistoryEdition/ActivityNode.h
@@ -2,11 +2,12 @@
 
 #pragma once
 
-#include "CoreMinimal.h"
 #include "ActivityDependencyEdge.h"
-#include "ActivityGraphIDs.h"
-#include "Misc/Optional.h"
-#include "Templates/NonNullPointer.h"
+#include "Containers/Array.h"
+#include "HistoryEdition/ActivityGraphIDs.h"
+#include "Misc/Optional.h" // IWYU pragma: keep
+
+struct FScriptContainerElement;
 
 namespace UE::ConcertSyncCore
 {
@@ -68,4 +69,10 @@ namespace UE::ConcertSyncCore
 		 */
 		TArray<FActivityNodeID> AffectedChildren;
 	};
-}
\ No newline at end of file
+}
+
+#if UE_ENABLE_INCLUDE_ORDER_DEPRECATED_IN_5_2
+#include "CoreMinimal.h"
+#include "Misc/Optional.h"
+#include "Templates/NonNullPointer.h"
+#endif
diff --git a/Engine/Plugins/Developer/Concert/ConcertSync/ConcertSyncCore/Source/ConcertSyncCore/Public/HistoryEdition/DebugDependencyGraph.h b/Engine/Plugins/Developer/Concert/ConcertSync/ConcertSyncCore/Source/ConcertSyncCore/Public/HistoryEdition/DebugDependencyGraph.h
--- a/Engine/Plugins/Developer/Concert/ConcertSync/ConcertSyncCore/Source/ConcertSyncCore/Public/HistoryEdition/DebugDependencyGraph.h
+++ b/Engine/Plugins/Developer/Concert/ConcertSync/ConcertSyncCore/Source/ConcertSyncCore/Public/HistoryEdition/DebugDependencyGraph.h
@@ -2,9 +2,14 @@
 
 #pragma once
 
-#include "CoreMinimal.h"
-#include "ActivityDependencyGraph.h"
-#include "ActivityGraphIDs.h"
+#include "Containers/StringFwd.h"
+#include "Misc/EnumClassFlags.h"
+
+class FString;
+namespace UE::ConcertSyncCore { class FActivityDependencyEdge; }
+namespace UE::ConcertSyncCore { class FActivityDependencyGraph; }
+namespace UE::ConcertSyncCore { struct FActivityNodeID; }
+template <typename FuncType> class TFunctionRef;
 
 class FConcertSyncSessionDatabase;
 
@@ -62,3 +67,9 @@ namespace UE::ConcertSyncCore::Graphviz
 	 */
 	CONCERTSYNCCORE_API FString ExportToGraphviz(const FActivityDependencyGraph& Graph, FMakeNodeTitle MakeNodeTitleFunc, FGetNodeStyle GetNodeStyleFunc, FGetEdgeStyle GetEdgeStyleFunc);
 }
+
+#if UE_ENABLE_INCLUDE_ORDER_DEPRECATED_IN_5_2
+#include "ActivityDependencyGraph.h"
+#include "ActivityGraphIDs.h"
+#include "CoreMinimal.h"
+#endif
diff --git a/Engine/Plugins/Developer/Concert/ConcertSync/ConcertSyncCore/Source/ConcertSyncCore/Public/HistoryEdition/DependencyGraphBuilder.h b/Engine/Plugins/Developer/Concert/ConcertSync/ConcertSyncCore/Source/ConcertSyncCore/Public/HistoryEdition/DependencyGraphBuilder.h
--- a/Engine/Plugins/Developer/Concert/ConcertSync/ConcertSyncCore/Source/ConcertSyncCore/Public/HistoryEdition/DependencyGraphBuilder.h
+++ b/Engine/Plugins/Developer/Concert/ConcertSync/ConcertSyncCore/Source/ConcertSyncCore/Public/HistoryEdition/DependencyGraphBuilder.h
@@ -2,8 +2,9 @@
 
 #pragma once
 
-#include "CoreMinimal.h"
-#include "ActivityDependencyGraph.h"
+#include "HAL/Platform.h"
+
+namespace UE::ConcertSyncCore { class FActivityDependencyGraph; }
 
 class FConcertSyncSessionDatabase;
 
@@ -14,3 +15,8 @@ namespace UE::ConcertSyncCore
 	 */
 	CONCERTSYNCCORE_API FActivityDependencyGraph BuildDependencyGraphFrom(const FConcertSyncSessionDatabase& SessionDatabase);
 }
+
+#if UE_ENABLE_INCLUDE_ORDER_DEPRECATED_IN_5_2
+#include "ActivityDependencyGraph.h"
+#include "CoreMinimal.h"
+#endif
diff --git a/Engine/Plugins/Developer/Concert/ConcertSync/ConcertSyncCore/Source/ConcertSyncCore/Public/HistoryEdition/HistoryAnalysis.h b/Engine/Plugins/Developer/Concert/ConcertSync/ConcertSyncCore/Source/ConcertSyncCore/Public/HistoryEdition/HistoryAnalysis.h
--- a/Engine/Plugins/Developer/Concert/ConcertSync/ConcertSyncCore/Source/ConcertSyncCore/Public/HistoryEdition/HistoryAnalysis.h
+++ b/Engine/Plugins/Developer/Concert/ConcertSync/ConcertSyncCore/Source/ConcertSyncCore/Public/HistoryEdition/HistoryAnalysis.h
@@ -2,10 +2,10 @@
 
 #pragma once
 
-#include "CoreMinimal.h"
-#include "ActivityDependencyGraph.h"
 #include "ConcertMessageData.h"
 
+namespace UE::ConcertSyncCore { class FActivityDependencyGraph; }
+
 class FConcertSyncSessionDatabase;
 
 namespace UE::ConcertSyncCore
@@ -55,3 +55,8 @@ namespace UE::ConcertSyncCore
 	CONCERTSYNCCORE_API FHistoryAnalysisResult AnalyseActivityDependencies_BottomUp(const TSet<FActivityID>& ActivitiesToEdit, const FActivityDependencyGraph& DependencyGraph, bool bAddEditedAsHardDependencies = false);
 }
 
+
+#if UE_ENABLE_INCLUDE_ORDER_DEPRECATED_IN_5_2
+#include "ActivityDependencyGraph.h"
+#include "CoreMinimal.h"
+#endif
diff --git a/Engine/Plugins/Developer/Concert/ConcertSync/ConcertSyncCore/Source/ConcertSyncCore/Public/HistoryEdition/HistoryEdition.h b/Engine/Plugins/Developer/Concert/ConcertSync/ConcertSyncCore/Source/ConcertSyncCore/Public/HistoryEdition/HistoryEdition.h
--- a/Engine/Plugins/Developer/Concert/ConcertSync/ConcertSyncCore/Source/ConcertSyncCore/Public/HistoryEdition/HistoryEdition.h
+++ b/Engine/Plugins/Developer/Concert/ConcertSync/ConcertSyncCore/Source/ConcertSyncCore/Public/HistoryEdition/HistoryEdition.h
@@ -2,8 +2,9 @@
 
 #pragma once
 
-#include "CoreMinimal.h"
-#include "HistoryAnalysis.h"
+#include "ConcertMessageData.h"
+
+namespace UE::ConcertSyncCore { struct FHistoryAnalysisResult; }
 
 class IConcertServer;
 class IConcertSyncServer;
@@ -28,3 +29,8 @@ namespace UE::ConcertSyncCore
 	/** Utility functions that converts FHistoryDeletionRequirements into a single TSet. */
 	CONCERTSYNCCORE_API TSet<FActivityID> CombineRequirements(const FHistoryAnalysisResult& ToDelete);
 }
+
+#if UE_ENABLE_INCLUDE_ORDER_DEPRECATED_IN_5_2
+#include "CoreMinimal.h"
+#include "HistoryAnalysis.h"
+#endif
diff --git a/Engine/Plugins/Developer/Concert/ConcertSync/ConcertSyncCore/Source/ConcertSyncCore/Public/IConcertFileSharingService.h b/Engine/Plugins/Developer/Concert/ConcertSync/ConcertSyncCore/Source/ConcertSyncCore/Public/IConcertFileSharingService.h
--- a/Engine/Plugins/Developer/Concert/ConcertSync/ConcertSyncCore/Source/ConcertSyncCore/Public/IConcertFileSharingService.h
+++ b/Engine/Plugins/Developer/Concert/ConcertSync/ConcertSyncCore/Source/ConcertSyncCore/Public/IConcertFileSharingService.h
@@ -2,7 +2,7 @@
 
 #pragma once
 
-#include "CoreMinimal.h"
+#include "Templates/SharedPointer.h"
 
 /**
  * Implements a file sharing service used as a side channel to ConcertTransport to exchange large files. This API is designed to
@@ -62,3 +62,7 @@ public:
 	 */
 	virtual TSharedPtr<FArchive> CreateReader(const FString& InFileId) = 0;
 };
+
+#if UE_ENABLE_INCLUDE_ORDER_DEPRECATED_IN_5_2
+#include "CoreMinimal.h"
+#endif
diff --git a/Engine/Plugins/Editor/ChangelistReview/Source/Private/SSourceControlReview.cpp b/Engine/Plugins/Editor/ChangelistReview/Source/Private/SSourceControlReview.cpp
--- a/Engine/Plugins/Editor/ChangelistReview/Source/Private/SSourceControlReview.cpp
+++ b/Engine/Plugins/Editor/ChangelistReview/Source/Private/SSourceControlReview.cpp
@@ -3,12 +3,9 @@
 #include "SSourceControlReview.h"
 
 #include "AssetRegistry/AssetData.h"
-#include "HAL/FileManager.h"
 #include "ClassIconFinder.h"
 #include "Framework/Views/TableViewMetadata.h"
 #include "ISourceControlModule.h"
-#include "SourceControlHelpers.h"
-#include "Misc/PackageName.h"
 #include "SourceControlOperations.h"
 #include "SSourceControlReviewEntry.h"
 #include "AssetRegistry/AssetRegistryModule.h"
diff --git a/Engine/Plugins/Editor/ContentBrowser/ContentBrowserAssetDataSource/Source/ContentBrowserAssetDataSource/Private/AssetFileContextMenu.cpp b/Engine/Plugins/Editor/ContentBrowser/ContentBrowserAssetDataSource/Source/ContentBrowserAssetDataSource/Private/AssetFileContextMenu.cpp
--- a/Engine/Plugins/Editor/ContentBrowser/ContentBrowserAssetDataSource/Source/ContentBrowserAssetDataSource/Private/AssetFileContextMenu.cpp
+++ b/Engine/Plugins/Editor/ContentBrowser/ContentBrowserAssetDataSource/Source/ContentBrowserAssetDataSource/Private/AssetFileContextMenu.cpp
@@ -4,7 +4,6 @@
 
 #include "AssetDefinition.h"
 #include "AssetDefinitionRegistry.h"
-#include "AssetRegistry/ARFilter.h"
 #include "Engine/Blueprint.h"
 #include "Engine/UserDefinedStruct.h"
 #include "Engine/UserDefinedEnum.h"
diff --git a/Engine/Plugins/Editor/ContentBrowser/ContentBrowserAssetDataSource/Source/ContentBrowserAssetDataSource/Private/ContentBrowserAssetDataCore.cpp b/Engine/Plugins/Editor/ContentBrowser/ContentBrowserAssetDataSource/Source/ContentBrowserAssetDataSource/Private/ContentBrowserAssetDataCore.cpp
--- a/Engine/Plugins/Editor/ContentBrowser/ContentBrowserAssetDataSource/Source/ContentBrowserAssetDataSource/Private/ContentBrowserAssetDataCore.cpp
+++ b/Engine/Plugins/Editor/ContentBrowser/ContentBrowserAssetDataSource/Source/ContentBrowserAssetDataSource/Private/ContentBrowserAssetDataCore.cpp
@@ -3,7 +3,6 @@
 #include "ContentBrowserAssetDataCore.h"
 
 #include "AssetDefinition.h"
-#include "AssetRegistry/ARFilter.h"
 #include "ContentBrowserDataSource.h"
 #include "ContentBrowserAssetDataPayload.h"
 #include "IAssetTools.h"
diff --git a/Engine/Plugins/Editor/DataValidation/Source/DataValidation/Public/DataValidationChangelist.h b/Engine/Plugins/Editor/DataValidation/Source/DataValidation/Public/DataValidationChangelist.h
--- a/Engine/Plugins/Editor/DataValidation/Source/DataValidation/Public/DataValidationChangelist.h
+++ b/Engine/Plugins/Editor/DataValidation/Source/DataValidation/Public/DataValidationChangelist.h
@@ -2,13 +2,7 @@
 
 #pragma once
 
-#include "EditorSubsystem.h"
 
-#include "UObject/ObjectMacros.h"
-#include "UObject/UObjectGlobals.h"
-#include "UObject/Interface.h"
-#include "UObject/ObjectMacros.h"
-#include "Engine/EngineTypes.h"
 
 #include "ISourceControlChangelist.h"
 
@@ -40,4 +34,10 @@ public:
 
 	/** Change description */
 	//FText Description;
-};
\ No newline at end of file
+};
+
+#if UE_ENABLE_INCLUDE_ORDER_DEPRECATED_IN_5_2
+#include "EditorSubsystem.h"
+#include "Engine/EngineTypes.h"
+#include "UObject/Interface.h"
+#endif
diff --git a/Engine/Plugins/Editor/DataValidation/Source/DataValidation/Public/DataValidationCommandlet.h b/Engine/Plugins/Editor/DataValidation/Source/DataValidation/Public/DataValidationCommandlet.h
--- a/Engine/Plugins/Editor/DataValidation/Source/DataValidation/Public/DataValidationCommandlet.h
+++ b/Engine/Plugins/Editor/DataValidation/Source/DataValidation/Public/DataValidationCommandlet.h
@@ -2,7 +2,6 @@
 #pragma once
 
 #include "Commandlets/Commandlet.h"
-#include "UObject/Interface.h"
 #include "DataValidationCommandlet.generated.h"
 
 UCLASS(CustomConstructor)
@@ -24,3 +23,7 @@ public:
 	// do the validation without creating a commandlet
 	static bool ValidateData(const FString& FullCommandLine);
 };
+
+#if UE_ENABLE_INCLUDE_ORDER_DEPRECATED_IN_5_2
+#include "UObject/Interface.h"
+#endif
diff --git a/Engine/Plugins/Editor/DataValidation/Source/DataValidation/Public/DataValidationManager.h b/Engine/Plugins/Editor/DataValidation/Source/DataValidation/Public/DataValidationManager.h
--- a/Engine/Plugins/Editor/DataValidation/Source/DataValidation/Public/DataValidationManager.h
+++ b/Engine/Plugins/Editor/DataValidation/Source/DataValidation/Public/DataValidationManager.h
@@ -2,10 +2,7 @@
 
 #pragma once
 
-#include "CoreMinimal.h"
-#include "UObject/Object.h"
-#include "AssetRegistry/AssetData.h"
-#include "Engine/EngineTypes.h"
+#include "UObject/Package.h"
 #include "DataValidationManager.generated.h"
 
 class FLogWindowManager;
@@ -110,3 +107,9 @@ private:
 	UPROPERTY(config)
 	FSoftClassPath DataValidationManagerClassName;
 };
+
+#if UE_ENABLE_INCLUDE_ORDER_DEPRECATED_IN_5_2
+#include "AssetRegistry/AssetData.h"
+#include "CoreMinimal.h"
+#include "Engine/EngineTypes.h"
+#endif
diff --git a/Engine/Plugins/Editor/DataValidation/Source/DataValidation/Public/DataValidationModule.h b/Engine/Plugins/Editor/DataValidation/Source/DataValidation/Public/DataValidationModule.h
--- a/Engine/Plugins/Editor/DataValidation/Source/DataValidation/Public/DataValidationModule.h
+++ b/Engine/Plugins/Editor/DataValidation/Source/DataValidation/Public/DataValidationModule.h
@@ -2,11 +2,10 @@
 
 #pragma once
 
-#include "CoreMinimal.h"
-#include "Stats/Stats.h"
-#include "AssetRegistry/AssetData.h"
-#include "Modules/ModuleInterface.h"
 #include "Modules/ModuleManager.h"
+#include "UObject/ObjectMacros.h"
+
+struct FAssetData;
 
 #include "DataValidationModule.generated.h"
 
@@ -67,3 +66,8 @@ public:
 	/** Validates selected assets and opens a window to report the results. If bValidateDependencies is true it will also validate any assets that the selected assets depend on. */
 	virtual void ValidateAssets(const TArray<FAssetData>& SelectedAssets, bool bValidateDependencies, const EDataValidationUsecase InValidationUsecase) = 0;
 };
+
+#if UE_ENABLE_INCLUDE_ORDER_DEPRECATED_IN_5_2
+#include "AssetRegistry/AssetData.h"
+#include "CoreMinimal.h"
+#endif
