Updated Concert archives to serialize all objects as FSoftObjectPath internally, to avoid needing to create FNames for full object paths

#rnx

diff --git a/Engine/Plugins/Developer/Concert/ConcertMain/Source/ConcertTransport/Private/IdentifierTable/ConcertTransportArchives.cpp b/Engine/Plugins/Developer/Concert/ConcertMain/Source/ConcertTransport/Private/IdentifierTable/ConcertTransportArchives.cpp
--- a/Engine/Plugins/Developer/Concert/ConcertMain/Source/ConcertTransport/Private/IdentifierTable/ConcertTransportArchives.cpp
+++ b/Engine/Plugins/Developer/Concert/ConcertMain/Source/ConcertTransport/Private/IdentifierTable/ConcertTransportArchives.cpp
@@ -133,25 +133,6 @@ void ReadName(FArchive& Ar, FName& Name, const FConcertLocalIdentifierTable* Loc
 	Name.SetNumber(NameNumber);
 }
 
-void WriteSoftObjectPath(FArchive& Ar, const FSoftObjectPath& AssetPtr)
-{
-	// TODO: Serialize via FSoftObjectPath::SerializePath to avoid converting all object paths to FName
-	FNameBuilder ObjPathStr;
-	AssetPtr.ToString(ObjPathStr);
-
-	FName ObjPath = FName(ObjPathStr);
-	Ar << ObjPath;
-}
-
-void ReadSoftObjectPath(FArchive& Ar, FSoftObjectPath& AssetPtr)
-{
-	FName ObjPath;
-	Ar << ObjPath;
-
-	FNameBuilder ObjPathStr(ObjPath);
-	AssetPtr.SetPath(ObjPathStr.ToView());
-}
-
 } // namespace UE::Concert::Private::ConcertIdentifierArchiveUtil
 
 FConcertIdentifierWriter::FConcertIdentifierWriter(FConcertLocalIdentifierTable* InLocalIdentifierTable, TArray<uint8>& InBytes, bool bIsPersistent)
@@ -168,7 +149,7 @@ FArchive& FConcertIdentifierWriter::operator<<(FName& Name)
 
 FArchive& FConcertIdentifierWriter::operator<<(FSoftObjectPath& AssetPtr)
 {
-	UE::Concert::Private::ConcertIdentifierArchiveUtil::WriteSoftObjectPath(*this, AssetPtr);
+	AssetPtr.SerializePath(*this);
 	return *this;
 }
 
@@ -197,7 +178,7 @@ FArchive& FConcertIdentifierReader::operator<<(FName& Name)
 
 FArchive& FConcertIdentifierReader::operator<<(FSoftObjectPath& AssetPtr)
 {
-	UE::Concert::Private::ConcertIdentifierArchiveUtil::ReadSoftObjectPath(*this, AssetPtr);
+	AssetPtr.SerializePath(*this);
 	return *this;
 }
 
@@ -256,7 +237,7 @@ FArchive& FConcertIdentifierRewriter::operator<<(FName& Name)
 
 FArchive& FConcertIdentifierRewriter::operator<<(FSoftObjectPath& AssetPtr)
 {
-	UE::Concert::Private::ConcertIdentifierArchiveUtil::WriteSoftObjectPath(*this, AssetPtr);
+	AssetPtr.SerializePath(*this);
 	return *this;
 }
 
diff --git a/Engine/Plugins/Developer/Concert/ConcertSync/ConcertSyncCore/Source/ConcertSyncCore/Private/ConcertSyncArchives.cpp b/Engine/Plugins/Developer/Concert/ConcertSync/ConcertSyncCore/Source/ConcertSyncCore/Private/ConcertSyncArchives.cpp
--- a/Engine/Plugins/Developer/Concert/ConcertSync/ConcertSyncCore/Source/ConcertSyncCore/Private/ConcertSyncArchives.cpp
+++ b/Engine/Plugins/Developer/Concert/ConcertSync/ConcertSyncCore/Source/ConcertSyncCore/Private/ConcertSyncArchives.cpp
@@ -9,7 +9,7 @@
 #include "Serialization/ObjectWriter.h"
 #include "UObject/Package.h"
 
-static const FName SkipAssetsMarker = TEXT("SKIPASSETS");
+static const FSoftObjectPath SkipAssetsMarker = FSoftObjectPath(TEXT("/Engine/Transient.__CONCERT_SKIP_ASSETS__"));
 
 namespace ConcertSyncUtil
 {
@@ -267,7 +267,7 @@ void FConcertSyncObjectWriter::SerializeProperty(const FProperty* InProp, const
 
 FArchive& FConcertSyncObjectWriter::operator<<(UObject*& Obj)
 {
-	FName ObjPath;
+	FSoftObjectPath ObjPath;
 	if (Obj)
 	{
 		if (bSkipAssets && Obj->IsAsset())
@@ -278,11 +278,10 @@ FArchive& FConcertSyncObjectWriter::operator<<(UObject*& Obj)
 		{
 			FString ObjectPathString = Obj->GetPathName();
 			RemapObjectPathDelegate.ExecuteIfBound(ObjectPathString);
-			ObjPath = FName(ObjectPathString);
+			ObjPath.SetPath(ObjectPathString);
 		}
 	}
-
-	*this << ObjPath;
+	ObjPath.SerializePath(*this);
 	return *this;
 }
 
@@ -311,8 +310,7 @@ FArchive& FConcertSyncObjectWriter::operator<<(FSoftObjectPtr& AssetPtr)
 
 FArchive& FConcertSyncObjectWriter::operator<<(FSoftObjectPath& AssetPtr)
 {
-	// TODO: Serialize via FSoftObjectPath::SerializePath to avoid converting all object paths to FName
-	FName ObjPath;
+	FSoftObjectPath ObjPath;
 	if (bSkipAssets)
 	{
 		ObjPath = SkipAssetsMarker;
@@ -321,9 +319,9 @@ FArchive& FConcertSyncObjectWriter::operator<<(FSoftObjectPath& AssetPtr)
 	{
 		FString ObjectPathString = AssetPtr.ToString();
 		RemapObjectPathDelegate.ExecuteIfBound(ObjectPathString);
-		ObjPath = FName(ObjectPathString);
+		ObjPath.SetPath(ObjectPathString);
 	}
-	*this << ObjPath;
+	ObjPath.SerializePath(*this);
 	return *this;
 }
 
@@ -415,16 +413,17 @@ void FConcertSyncObjectReader::SerializeProperty(const FProperty* InProp, UObjec
 
 FArchive& FConcertSyncObjectReader::operator<<(UObject*& Obj)
 {
-	FName ObjPath;
-	*this << ObjPath;
+	FSoftObjectPath ObjPath;
+	ObjPath.SerializePath(*this);
 
-	if (ObjPath.IsNone())
+	if (ObjPath.IsNull())
 	{
 		Obj = nullptr;
 	}
 	else if (ObjPath != SkipAssetsMarker)
 	{
 		const FString ResolvedObjPath = WorldRemapper.RemapObjectPathName(ObjPath.ToString());
+		OnObjectSerialized(FSoftObjectPath(ObjPath));
 
 		// Always attempt to find an in-memory object first as we may be calling this function while a load is taking place
 		Obj = StaticFindObject(UObject::StaticClass(), nullptr, *ResolvedObjPath);
@@ -493,13 +492,15 @@ FArchive& FConcertSyncObjectReader::operator<<(FSoftObjectPtr& AssetPtr)
 
 FArchive& FConcertSyncObjectReader::operator<<(FSoftObjectPath& AssetPtr)
 {
-	FName ObjPath;
-	*this << ObjPath;
+	FSoftObjectPath ObjPath;
+	ObjPath.SerializePath(*this);
 
 	if (ObjPath != SkipAssetsMarker)
 	{
 		const FString ResolvedObjPath = WorldRemapper.RemapObjectPathName(ObjPath.ToString());
 		AssetPtr.SetPath(ResolvedObjPath);
+
+		OnObjectSerialized(AssetPtr);
 	}
 
 	return *this;
@@ -549,10 +550,10 @@ void FConcertSyncObjectRewriter::RewriteProperty(const FProperty* InProp)
 
 FArchive& FConcertSyncObjectRewriter::operator<<(UObject*& Obj)
 {
-	FName ObjPath;
-	*this << ObjPath;
+	FSoftObjectPath ObjPath;
+	ObjPath.SerializePath(*this);
 
-	if (ObjPath.IsNone())
+	if (ObjPath.IsNull())
 	{
 		Obj = nullptr;
 	}
diff --git a/Engine/Plugins/Developer/Concert/ConcertSync/ConcertSyncCore/Source/ConcertSyncCore/Public/ConcertSyncArchives.h b/Engine/Plugins/Developer/Concert/ConcertSync/ConcertSyncCore/Source/ConcertSyncCore/Public/ConcertSyncArchives.h
--- a/Engine/Plugins/Developer/Concert/ConcertSync/ConcertSyncCore/Source/ConcertSyncCore/Public/ConcertSyncArchives.h
+++ b/Engine/Plugins/Developer/Concert/ConcertSync/ConcertSyncCore/Source/ConcertSyncCore/Public/ConcertSyncArchives.h
@@ -101,6 +101,9 @@ public:
 	virtual FString GetArchiveName() const override;
 	//~ End FArchive Interface
 
+protected:
+	virtual void OnObjectSerialized(const FSoftObjectPath& Obj) {}
+
 private:
 	FConcertSyncWorldRemapper WorldRemapper;
 	FConcertSyncEncounteredMissingObject EncounteredMissingObjectDelegate;
