MediaPlate: Ensure the transaction of the internally created material instance constant.

diff --git a/Engine/Plugins/Developer/Concert/ConcertSync/ConcertSyncCore/Config/BaseConcertSyncCore.ini b/Engine/Plugins/Developer/Concert/ConcertSync/ConcertSyncCore/Config/BaseConcertSyncCore.ini
--- a/Engine/Plugins/Developer/Concert/ConcertSync/ConcertSyncCore/Config/BaseConcertSyncCore.ini
+++ b/Engine/Plugins/Developer/Concert/ConcertSync/ConcertSyncCore/Config/BaseConcertSyncCore.ini
@@ -2,7 +2,7 @@
 +IncludeObjectClassFilters=(ObjectOuterClass=None,ObjectClasses=(Class'/Script/Engine.World',Class'/Script/Engine.BookmarkBase',Class'/Script/Engine.MaterialParameterCollection',Class'/Script/MovieScene.MovieSceneSignedObject',Class'/Script/MovieScene.MovieSceneFolder'))
 +IncludeObjectClassFilters=(ObjectOuterClass=Class'/Script/Engine.World',ObjectClasses=(Class'/Script/Engine.Layer',Class'/Script/Engine.Level',Class'/Script/Engine.LevelStreaming',Class'/Script/Engine.Actor',Class'/Script/Engine.ActorComponent',Class'/Script/Engine.Polys',Class'/Script/Engine.BrushBuilder',Class'/Script/Foliage.FoliageType'))
 +IncludeObjectClassFilters=(ObjectOuterClass=Class'/Script/MovieScene.MovieSceneSignedObject',ObjectClasses=(Class'/Script/Engine.Actor',Class'/Script/Engine.ActorComponent'))
-+IncludeObjectClassFilters=(ObjectOuterClass=Class'/Script/Engine.World',ObjectClasses=(Class'/Script/Engine.MaterialInstanceDynamic', '/Script/MediaAssets.MediaPlaylist'))
++IncludeObjectClassFilters=(ObjectOuterClass=Class'/Script/Engine.World',ObjectClasses=(Class'/Script/Engine.MaterialInstanceDynamic', '/Script/Engine.MaterialInstanceConstant', '/Script/MediaAssets.MediaPlaylist'))
 +IncludeObjectClassFilters=(ObjectOuterClass=Class'/Script/Engine.World',ObjectClasses=(Class'/Script/Constraints.TickableConstraint', '/Script/Constraints.TransformableHandle', '/Script/Constraints.ConstraintsActor', '/Script/Constraints.ConstraintsManager'))
 +ExcludeTransactionClassFilters=(ObjectOuterClass=None,ObjectClasses=(Class'/Script/Landscape.LandscapeProxy',Class'/Script/Landscape.LandscapeHeightfieldCollisionComponent',Class'/Script/Landscape.LandscapeComponent'))
 +AllowedTransientProperties="/Script/Engine.Actor:bHiddenEdTemporary"
diff --git a/Engine/Plugins/Media/MediaPlate/Source/MediaPlate/Private/MediaPlate.cpp b/Engine/Plugins/Media/MediaPlate/Source/MediaPlate/Private/MediaPlate.cpp
--- a/Engine/Plugins/Media/MediaPlate/Source/MediaPlate/Private/MediaPlate.cpp
+++ b/Engine/Plugins/Media/MediaPlate/Source/MediaPlate/Private/MediaPlate.cpp
@@ -182,7 +182,7 @@ void AMediaPlate::ApplyMaterial(UMaterialInterface* Material)
 
 				// Create instance.
 				UMaterialInstanceConstant* MaterialInstance =
-					NewObject<UMaterialInstanceConstant>(StaticMeshComponent, MaterialUniqueName);
+					NewObject<UMaterialInstanceConstant>(StaticMeshComponent, MaterialUniqueName, RF_Transactional);
 				MaterialInstance->SetParentEditorOnly(Material);
 				MaterialInstance->CopyMaterialUniformParametersEditorOnly(Material);
 				MaterialInstance->SetTextureParameterValueEditorOnly(
