Fixed mismatched copy constructor/assignment operator in UE::ConcertSyncCore::FActivityNodeID.

diff --git a/Engine/Plugins/Developer/Concert/ConcertSync/ConcertSyncCore/Source/ConcertSyncCore/Public/HistoryEdition/ActivityGraphIDs.h b/Engine/Plugins/Developer/Concert/ConcertSync/ConcertSyncCore/Source/ConcertSyncCore/Public/HistoryEdition/ActivityGraphIDs.h
--- a/Engine/Plugins/Developer/Concert/ConcertSync/ConcertSyncCore/Source/ConcertSyncCore/Public/HistoryEdition/ActivityGraphIDs.h
+++ b/Engine/Plugins/Developer/Concert/ConcertSync/ConcertSyncCore/Source/ConcertSyncCore/Public/HistoryEdition/ActivityGraphIDs.h
@@ -20,10 +20,6 @@ namespace UE::ConcertSyncCore
 			: ID(Forward<T>(ID))
 		{}
 
-		FActivityNodeID(const FActivityNodeID& Other)
-			: ID(Other.ID)
-		{}
-
 		/** Conversion back to to size_t is not a common programmer mistake */
 		explicit operator size_t() const { return ID; }
 
