Remove landscape from the excluded transaction list. It was disabled
initially because landscape support was still experimental.  With OFPA /
WP and PCG features coming out it only make sense to turn this on
official and deal with the fall-out.

diff --git a/Engine/Plugins/Developer/Concert/ConcertSync/ConcertSyncCore/Config/BaseConcertSyncCore.ini b/Engine/Plugins/Developer/Concert/ConcertSync/ConcertSyncCore/Config/BaseConcertSyncCore.ini
--- a/Engine/Plugins/Developer/Concert/ConcertSync/ConcertSyncCore/Config/BaseConcertSyncCore.ini
+++ b/Engine/Plugins/Developer/Concert/ConcertSync/ConcertSyncCore/Config/BaseConcertSyncCore.ini
@@ -6,7 +6,7 @@
 +IncludeObjectClassFilters=(ObjectOuterClass=Class'/Script/MovieScene.MovieSceneSignedObject',ObjectClasses=(Class'/Script/Engine.Actor',Class'/Script/Engine.ActorComponent'))
 +IncludeObjectClassFilters=(ObjectOuterClass=Class'/Script/Engine.World',ObjectClasses=(Class'/Script/Engine.MaterialInstanceDynamic', '/Script/Engine.MaterialInstanceConstant', '/Script/MediaAssets.MediaPlaylist'))
 +IncludeObjectClassFilters=(ObjectOuterClass=Class'/Script/Engine.World',ObjectClasses=(Class'/Script/Constraints.TickableConstraint', '/Script/Constraints.TransformableHandle', '/Script/Constraints.ConstraintsActor', '/Script/Constraints.ConstraintsManager'))
-+ExcludeTransactionClassFilters=(ObjectOuterClass=None,ObjectClasses=(Class'/Script/Landscape.LandscapeProxy',Class'/Script/Landscape.LandscapeHeightfieldCollisionComponent',Class'/Script/Landscape.LandscapeComponent'))
+
 +AllowedTransientProperties="/Script/Engine.Actor:bHiddenEdTemporary"
 +AllowedTransientProperties="/Script/Engine.Actor:bHiddenEdLevel"
 +AllowedTransientProperties="/Script/Engine.Level:bIsVisible"
