Fix deferencing nullptr found by static analysis.

#rnx

diff --git a/Engine/Plugins/Developer/Concert/ConcertSync/ConcertSyncCore/Source/ConcertSyncCore/Private/ConcertSyncArchives.cpp b/Engine/Plugins/Developer/Concert/ConcertSync/ConcertSyncCore/Source/ConcertSyncCore/Private/ConcertSyncArchives.cpp
--- a/Engine/Plugins/Developer/Concert/ConcertSync/ConcertSyncCore/Source/ConcertSyncCore/Private/ConcertSyncArchives.cpp
+++ b/Engine/Plugins/Developer/Concert/ConcertSync/ConcertSyncCore/Source/ConcertSyncCore/Private/ConcertSyncArchives.cpp
@@ -309,6 +309,7 @@ FArchive& FConcertSyncObjectWriter::operator<<(UObject*& Obj)
 
 	bool bDoSerializeObject = bCanSerializeObject
 		&& bSerializeNestedObjects
+		&& Obj && Obj->GetPackage()
 		&& PackageName == Obj->GetPackage()->GetPathName()
 		&& !Obj->IsA<UPackage>();
 
@@ -318,6 +319,8 @@ FArchive& FConcertSyncObjectWriter::operator<<(UObject*& Obj)
 		// To prevent stack overflow we need to add to the collected objects first.
 		CollectedObjects.Add(ObjPath.ToString());
 
+		check(Obj && Obj->GetClass());
+
 		FString ClassName = Obj->GetClass()->GetPathName();
 		int32 OutFlags = static_cast<int32>(Obj->GetFlags());
 		*this << ClassName;
