Only serialize our nested object flag when we are running the reader/writer in those modes.

diff --git a/Engine/Plugins/Developer/Concert/ConcertSync/ConcertSyncClient/Source/ConcertSyncClient/Private/ConcertClientSequencerManager.cpp b/Engine/Plugins/Developer/Concert/ConcertSync/ConcertSyncClient/Source/ConcertSyncClient/Private/ConcertClientSequencerManager.cpp
--- a/Engine/Plugins/Developer/Concert/ConcertSync/ConcertSyncClient/Source/ConcertSyncClient/Private/ConcertClientSequencerManager.cpp
+++ b/Engine/Plugins/Developer/Concert/ConcertSync/ConcertSyncClient/Source/ConcertSyncClient/Private/ConcertClientSequencerManager.cpp
@@ -755,6 +755,7 @@ void FConcertClientSequencerManager::ApplyOpenEvent(const FConcertSequencerOpenE
 			TEXT("/Engine/Transient.__PendingLevelSequence__"), PendingLevelSequence->GetPathName());
 		FConcertSyncObjectReader Reader(nullptr, MoveTemp(Remapper), nullptr, PendingLevelSequence, InOpenEvent.TakeData.Bytes,
 										MissingObjectDelegate);
+		Reader.SetSerializeNestedObjects(true);
 		Reader.SerializeObject(PendingLevelSequence);
 		#if WITH_EDITOR
 		GEditor->GetEditorSubsystem<UAssetEditorSubsystem>()->OpenEditorForAsset(SequenceObjectPath);
diff --git a/Engine/Plugins/Developer/Concert/ConcertSync/ConcertSyncCore/Source/ConcertSyncCore/Private/ConcertSyncArchives.cpp b/Engine/Plugins/Developer/Concert/ConcertSync/ConcertSyncCore/Source/ConcertSyncCore/Private/ConcertSyncArchives.cpp
--- a/Engine/Plugins/Developer/Concert/ConcertSync/ConcertSyncCore/Source/ConcertSyncCore/Private/ConcertSyncArchives.cpp
+++ b/Engine/Plugins/Developer/Concert/ConcertSync/ConcertSyncCore/Source/ConcertSyncCore/Private/ConcertSyncArchives.cpp
@@ -313,9 +313,14 @@ FArchive& FConcertSyncObjectWriter::operator<<(UObject*& Obj)
 		&& PackageName == Obj->GetPackage()->GetPathName()
 		&& !Obj->IsA<UPackage>();
 
-	*this << bDoSerializeObject;
+	if (bSerializeNestedObjects)
+	{
+		*this << bDoSerializeObject;
+	}
+
 	if (bDoSerializeObject)
 	{
+
 		// To prevent stack overflow we need to add to the collected objects first.
 		CollectedObjects.Add(ObjPath.ToString());
 
@@ -463,8 +468,12 @@ FArchive& FConcertSyncObjectReader::operator<<(UObject*& Obj)
 	FSoftObjectPath ObjPath;
 	ObjPath.SerializePath(*this);
 
-	bool bDidSerializeObject;
-	*this << bDidSerializeObject;
+	bool bDidSerializeObject = false;
+	if (bSerializeNestedObjects)
+	{
+		*this << bDidSerializeObject;
+	}
+
 	if (ObjPath.IsNull())
 	{
 		Obj = nullptr;
diff --git a/Engine/Plugins/Developer/Concert/ConcertSync/ConcertSyncCore/Source/ConcertSyncCore/Public/ConcertSyncArchives.h b/Engine/Plugins/Developer/Concert/ConcertSync/ConcertSyncCore/Source/ConcertSyncCore/Public/ConcertSyncArchives.h
--- a/Engine/Plugins/Developer/Concert/ConcertSync/ConcertSyncCore/Source/ConcertSyncCore/Public/ConcertSyncArchives.h
+++ b/Engine/Plugins/Developer/Concert/ConcertSync/ConcertSyncCore/Source/ConcertSyncCore/Public/ConcertSyncArchives.h
@@ -109,10 +109,15 @@ public:
 	virtual FString GetArchiveName() const override;
 	//~ End FArchive Interface
 
+	void SetSerializeNestedObjects(bool bInSerializeNestedObjects)
+	{
+		bSerializeNestedObjects = bInSerializeNestedObjects;
+	}
 protected:
 	virtual void OnObjectSerialized(const FSoftObjectPath& Obj) {}
 
 private:
+	bool bSerializeNestedObjects = false;
 	UObject* CurrentOuter = nullptr;
 	FConcertSyncWorldRemapper WorldRemapper;
 	FConcertSyncEncounteredMissingObject EncounteredMissingObjectDelegate;
