Fix missing include for non-unity

diff --git a/Engine/Plugins/Developer/Concert/ConcertSync/ConcertSyncCore/Source/ConcertSyncCore/Public/Replication/Formats/FullObjectReplicationData.h b/Engine/Plugins/Developer/Concert/ConcertSync/ConcertSyncCore/Source/ConcertSyncCore/Public/Replication/Formats/FullObjectReplicationData.h
--- a/Engine/Plugins/Developer/Concert/ConcertSync/ConcertSyncCore/Source/ConcertSyncCore/Public/Replication/Formats/FullObjectReplicationData.h
+++ b/Engine/Plugins/Developer/Concert/ConcertSync/ConcertSyncCore/Source/ConcertSyncCore/Public/Replication/Formats/FullObjectReplicationData.h
@@ -2,6 +2,7 @@
 
 #pragma once
 
+#include "ConcertMessageData.h"
 #include "FullObjectReplicationData.generated.h"
 
 /**
