Add an exclusion list to exlude certain properties that should not be included in a Multi-user transaction.  This will prevent trying to re-create an abstract class on a nDisplay instance thus crashing.

diff --git a/Engine/Plugins/Developer/Concert/ConcertSync/ConcertSyncCore/Config/BaseConcertSyncCore.ini b/Engine/Plugins/Developer/Concert/ConcertSync/ConcertSyncCore/Config/BaseConcertSyncCore.ini
--- a/Engine/Plugins/Developer/Concert/ConcertSync/ConcertSyncCore/Config/BaseConcertSyncCore.ini
+++ b/Engine/Plugins/Developer/Concert/ConcertSync/ConcertSyncCore/Config/BaseConcertSyncCore.ini
@@ -11,6 +11,7 @@
 +AllowedTransientProperties="/Script/Engine.Actor:bHiddenEdLevel"
 +AllowedTransientProperties="/Script/Engine.Level:bIsVisible"
 +AllowedTransientProperties="/Script/Engine.SceneComponent:AttachChildren"
++ExcludedProperties="/Script/Engine.InstanceCacheDataBase:UniqueTransientPackage"
 
 ; If you wish to use Composure in Multi-user, you may add that line under [/Script/ConcertSyncCore.ConcertSyncConfig] in your project engine config settings.
 ;+IncludeObjectClassFilters=(ObjectOuterClass=None,ObjectClasses=(Class'/Script/Composure.CompositingElementPass',Class'/Script/Composure.ComposurePostProcessPassPolicy'))
diff --git a/Engine/Plugins/Developer/Concert/ConcertSync/ConcertSyncCore/Source/ConcertSyncCore/Private/ConcertSyncArchives.cpp b/Engine/Plugins/Developer/Concert/ConcertSync/ConcertSyncCore/Source/ConcertSyncCore/Private/ConcertSyncArchives.cpp
--- a/Engine/Plugins/Developer/Concert/ConcertSync/ConcertSyncCore/Source/ConcertSyncCore/Private/ConcertSyncArchives.cpp
+++ b/Engine/Plugins/Developer/Concert/ConcertSync/ConcertSyncCore/Source/ConcertSyncCore/Private/ConcertSyncArchives.cpp
@@ -16,19 +16,19 @@ namespace ConcertSyncUtil
 
 bool CanExportProperty(const FProperty* Property, const bool InIncludeEditorOnlyData)
 {
-	auto CanExportTransientProperty = [Property]()
+	auto PropertyIsInList = [Property](const TArray<TFieldPath<FProperty>>& PropertyList)
 	{
-		const UConcertSyncConfig* SyncConfig = GetDefault<UConcertSyncConfig>();
-		return SyncConfig->AllowedTransientProperties.ContainsByPredicate([Property](const TFieldPath<FProperty>& TransactionProperty)
+		return PropertyList.ContainsByPredicate([Property](const TFieldPath<FProperty>& PropertyFieldPath)
 		{
-			FProperty* FilterProperty = TransactionProperty.Get();
+			FProperty* FilterProperty = PropertyFieldPath.Get();
 			return Property == FilterProperty;
 		});
 	};
-
+	const UConcertSyncConfig* SyncConfig = GetDefault<UConcertSyncConfig>();
 	return (!Property->IsEditorOnlyProperty() || InIncludeEditorOnlyData)
 		&& (!Property->HasAnyPropertyFlags(CPF_NonTransactional))
-		&& (!Property->HasAnyPropertyFlags(CPF_Transient) || CanExportTransientProperty());
+		&& (!Property->HasAnyPropertyFlags(CPF_Transient) || PropertyIsInList(SyncConfig->AllowedTransientProperties))
+		&& (!PropertyIsInList(SyncConfig->ExcludedProperties));
 }
 
 void GatherDefaultSubobjectPaths(const UObject* Obj, TSet<FSoftObjectPath>& OutSubobjects)
diff --git a/Engine/Plugins/Developer/Concert/ConcertSync/ConcertSyncCore/Source/ConcertSyncCore/Public/ConcertSyncSettings.h b/Engine/Plugins/Developer/Concert/ConcertSync/ConcertSyncCore/Source/ConcertSyncCore/Public/ConcertSyncSettings.h
--- a/Engine/Plugins/Developer/Concert/ConcertSync/ConcertSyncCore/Source/ConcertSyncCore/Public/ConcertSyncSettings.h
+++ b/Engine/Plugins/Developer/Concert/ConcertSync/ConcertSyncCore/Source/ConcertSyncCore/Public/ConcertSyncSettings.h
@@ -103,6 +103,12 @@ public:
 	UPROPERTY(config, EditAnywhere, Category="Transaction Settings", meta=(AllowedClasses="/Script/CoreUObject.Property"))
 	TArray<TFieldPath<FProperty>> AllowedTransientProperties;
 
+	/**
+	 * Array of class properties that we should filtered out.
+	 */
+	UPROPERTY(config, EditAnywhere, Category="Transaction Settings", meta=(AllowedClasses="/Script/CoreUObject.Property"))
+	TArray<TFieldPath<FProperty>> ExcludedProperties;
+
 	/**
 	 * Array of package class filter.
 	 * Packages that matches those filters, will be excluded from propagating to the server when saved.
