Make sure we have enough properties to check in the property chain

diff --git a/Engine/Plugins/Developer/Concert/ConcertSync/ConcertSyncCore/Source/ConcertSyncCore/Private/Replication/Data/ConcertPropertySelection.cpp b/Engine/Plugins/Developer/Concert/ConcertSync/ConcertSyncCore/Source/ConcertSyncCore/Private/Replication/Data/ConcertPropertySelection.cpp
--- a/Engine/Plugins/Developer/Concert/ConcertSync/ConcertSyncCore/Source/ConcertSyncCore/Private/Replication/Data/ConcertPropertySelection.cpp
+++ b/Engine/Plugins/Developer/Concert/ConcertSync/ConcertSyncCore/Source/ConcertSyncCore/Private/Replication/Data/ConcertPropertySelection.cpp
@@ -121,7 +121,7 @@ bool FConcertPropertyChain::MatchesExactly(const FArchiveSerializedPropertyChain
 	const bool bLeafPropertiesMatch = !PathToProperty.IsEmpty() && IsInnerContainerProperty(LeafProperty)
 		// The only place FConcertPropertyChain contains inner container properties is at the end; it is named InternalContainerPropertyValueName.
 		// In that case the "real" leaf property is the owning container property.
-		? OptionalChainLength >= 1 && OptionalChain->GetPropertyFromStack(0)->GetFName() == PathToProperty[PathToProperty.Num() - 2] 
+		? (OptionalChainLength >= 1 && PathToProperty.Num() > 1) && OptionalChain->GetPropertyFromStack(0)->GetFName() == PathToProperty[PathToProperty.Num() - 2]
 		: LeafProperty.GetFName() == PathToProperty[PathToProperty.Num() - 1];
 	// No point continuing matching properties if leaf properties do not match.
 	if (!bLeafPropertiesMatch)
