Allow Concert to filter out certain types of properties

diff --git a/Engine/Plugins/Developer/Concert/ConcertSync/ConcertSyncCore/Source/ConcertSyncCore/Private/ConcertSyncArchives.cpp b/Engine/Plugins/Developer/Concert/ConcertSync/ConcertSyncCore/Source/ConcertSyncCore/Private/ConcertSyncArchives.cpp
--- a/Engine/Plugins/Developer/Concert/ConcertSync/ConcertSyncCore/Source/ConcertSyncCore/Private/ConcertSyncArchives.cpp
+++ b/Engine/Plugins/Developer/Concert/ConcertSync/ConcertSyncCore/Source/ConcertSyncCore/Private/ConcertSyncArchives.cpp
@@ -16,19 +16,30 @@ namespace ConcertSyncUtil
 
 bool CanExportProperty(const FProperty* Property, const bool InIncludeEditorOnlyData)
 {
-	auto PropertyIsInList = [Property](const TArray<TFieldPath<FProperty>>& PropertyList)
+	auto PropertyPathIsInList = [Property](const TArray<TFieldPath<FProperty>>& PropertyPaths)
 	{
-		return PropertyList.ContainsByPredicate([Property](const TFieldPath<FProperty>& PropertyFieldPath)
+		return PropertyPaths.ContainsByPredicate([Property](const TFieldPath<FProperty>& PropertyFieldPath)
 		{
-			FProperty* FilterProperty = PropertyFieldPath.Get();
+			const FProperty* FilterProperty = PropertyFieldPath.Get();
 			return Property == FilterProperty;
 		});
 	};
+
+	auto PropertyTypeIsInList = [Property](const TArray<FName>& PropertyTypes)
+	{
+		return PropertyTypes.ContainsByPredicate([Property](const FName PropertyType)
+		{
+			const FFieldClass* FilterPropertyClass = FFieldClass::GetNameToFieldClassMap().FindRef(PropertyType);
+			return FilterPropertyClass && Property->GetClass()->IsChildOf(FilterPropertyClass);
+		});
+	};
+
 	const UConcertSyncConfig* SyncConfig = GetDefault<UConcertSyncConfig>();
 	return (!Property->IsEditorOnlyProperty() || InIncludeEditorOnlyData)
 		&& (!Property->HasAnyPropertyFlags(CPF_NonTransactional))
-		&& (!Property->HasAnyPropertyFlags(CPF_Transient) || PropertyIsInList(SyncConfig->AllowedTransientProperties))
-		&& (!PropertyIsInList(SyncConfig->ExcludedProperties));
+		&& (!Property->HasAnyPropertyFlags(CPF_Transient) || PropertyPathIsInList(SyncConfig->AllowedTransientProperties))
+		&& (!PropertyPathIsInList(SyncConfig->ExcludedProperties))
+		&& (!PropertyTypeIsInList(SyncConfig->ExcludedPropertyTypes));
 }
 
 void GatherDefaultSubobjectPaths(const UObject* Obj, TSet<FSoftObjectPath>& OutSubobjects)
diff --git a/Engine/Plugins/Developer/Concert/ConcertSync/ConcertSyncCore/Source/ConcertSyncCore/Public/ConcertSyncSettings.h b/Engine/Plugins/Developer/Concert/ConcertSync/ConcertSyncCore/Source/ConcertSyncCore/Public/ConcertSyncSettings.h
--- a/Engine/Plugins/Developer/Concert/ConcertSync/ConcertSyncCore/Source/ConcertSyncCore/Public/ConcertSyncSettings.h
+++ b/Engine/Plugins/Developer/Concert/ConcertSync/ConcertSyncCore/Source/ConcertSyncCore/Public/ConcertSyncSettings.h
@@ -109,6 +109,12 @@ public:
 	UPROPERTY(config, EditAnywhere, Category="Transaction Settings", meta=(AllowedClasses="/Script/CoreUObject.Property"))
 	TArray<TFieldPath<FProperty>> ExcludedProperties;
 
+	/**
+	 * Array of property types that we should filtered out.
+	 */
+	UPROPERTY(config, EditAnywhere, Category="Transaction Settings")
+	TArray<FName> ExcludedPropertyTypes;
+
 	/**
 	 * Array of package class filter.
 	 * Packages that matches those filters, will be excluded from propagating to the server when saved.
