Mark missing assets as skipped assets in the archive so that they do not get applied within the editor.

diff --git a/Engine/Plugins/Developer/Concert/ConcertSync/ConcertSyncCore/Source/ConcertSyncCore/Private/ConcertSyncArchives.cpp b/Engine/Plugins/Developer/Concert/ConcertSync/ConcertSyncCore/Source/ConcertSyncCore/Private/ConcertSyncArchives.cpp
--- a/Engine/Plugins/Developer/Concert/ConcertSync/ConcertSyncCore/Source/ConcertSyncCore/Private/ConcertSyncArchives.cpp
+++ b/Engine/Plugins/Developer/Concert/ConcertSync/ConcertSyncCore/Source/ConcertSyncCore/Private/ConcertSyncArchives.cpp
@@ -14,6 +14,11 @@ static const FSoftObjectPath SkipAssetsMarker = FSoftObjectPath(TEXT("/Engine/Tr
 namespace ConcertSyncUtil
 {
 
+const FSoftObjectPath& GetSkipObjectPath()
+{
+	return SkipAssetsMarker;
+}
+
 bool CanExportProperty(const FProperty* Property, const bool InIncludeEditorOnlyData)
 {
 	auto PropertyPathIsInList = [Property](const TArray<TFieldPath<FProperty>>& PropertyPaths)
diff --git a/Engine/Plugins/Developer/Concert/ConcertSync/ConcertSyncCore/Source/ConcertSyncCore/Public/ConcertSyncArchives.h b/Engine/Plugins/Developer/Concert/ConcertSync/ConcertSyncCore/Source/ConcertSyncCore/Public/ConcertSyncArchives.h
--- a/Engine/Plugins/Developer/Concert/ConcertSync/ConcertSyncCore/Source/ConcertSyncCore/Public/ConcertSyncArchives.h
+++ b/Engine/Plugins/Developer/Concert/ConcertSync/ConcertSyncCore/Source/ConcertSyncCore/Public/ConcertSyncArchives.h
@@ -21,6 +21,8 @@ namespace ConcertSyncUtil
 CONCERTSYNCCORE_API bool CanExportProperty(const FProperty* Property, const bool InIncludeEditorOnlyData);
 
 CONCERTSYNCCORE_API void ResetObjectPropertiesToArchetypeValues(UObject* Object, const bool InIncludeEditorOnlyData);
+
+CONCERTSYNCCORE_API const FSoftObjectPath& GetSkipObjectPath();
 }
 
 /** Util to handle remapping objects within the source world to be the equivalent objects in the destination world */
