Network messages in Concert now are prefixed with FConcertReplication_
Split messages into separate files

diff --git a/Engine/Plugins/Developer/Concert/ConcertApp/MultiUserClient/Source/MultiUserClient/Private/Replication/Authority/AuthorityPolicy.cpp b/Engine/Plugins/Developer/Concert/ConcertApp/MultiUserClient/Source/MultiUserClient/Private/Replication/Authority/AuthorityPolicy.cpp
--- a/Engine/Plugins/Developer/Concert/ConcertApp/MultiUserClient/Source/MultiUserClient/Private/Replication/Authority/AuthorityPolicy.cpp
+++ b/Engine/Plugins/Developer/Concert/ConcertApp/MultiUserClient/Source/MultiUserClient/Private/Replication/Authority/AuthorityPolicy.cpp
@@ -28,7 +28,7 @@ namespace UE::MultiUserClient
 			}
 
 			// Call for put object requests 
-			for (const TPair<FObjectInStreamID, FConcertChangeStream_PutObject>& Pair : AcceptedRequest.ObjectsToPut)
+			for (const TPair<FObjectInStreamID, FConcertReplication_ChangeStream_PutObject>& Pair : AcceptedRequest.ObjectsToPut)
 			{
 				const bool bObjectWasAddedWithPut = !OldState.ReplicatedObjects.Contains(Pair.Key.Object);
 				if (bObjectWasAddedWithPut)
diff --git a/Engine/Plugins/Developer/Concert/ConcertApp/MultiUserClient/Source/MultiUserClient/Private/Replication/Stream/LocalClientStreamSynchronizer.cpp b/Engine/Plugins/Developer/Concert/ConcertApp/MultiUserClient/Source/MultiUserClient/Private/Replication/Stream/LocalClientStreamSynchronizer.cpp
--- a/Engine/Plugins/Developer/Concert/ConcertApp/MultiUserClient/Source/MultiUserClient/Private/Replication/Stream/LocalClientStreamSynchronizer.cpp
+++ b/Engine/Plugins/Developer/Concert/ConcertApp/MultiUserClient/Source/MultiUserClient/Private/Replication/Stream/LocalClientStreamSynchronizer.cpp
@@ -142,7 +142,7 @@ namespace UE::MultiUserClient
 			return EBreakBehavior::Continue;
 		});
 
-		FConcertChangeStream_Request Request = ConcertSyncCore::Replication::ChangeStreamUtils::BuildRequestFromDiff(StreamId, Base, Cleansed);
+		FConcertReplication_ChangeStream_Request Request = ConcertSyncCore::Replication::ChangeStreamUtils::BuildRequestFromDiff(StreamId, Base, Cleansed);
 		return { MoveTemp(Request.ObjectsToRemove), MoveTemp(Request.ObjectsToPut) };
 	}
 
@@ -180,7 +180,7 @@ namespace UE::MultiUserClient
 
 	ConcertSyncClient::Replication::FChangeStreamRequest FLocalClientStreamSynchronizer::BuildChangeRequest_CreateNewStream(const FGuid& StreamId, const FChangelist& FromChangelist)
 	{
-		const TMap<FObjectInStreamID, FConcertChangeStream_PutObject>& ObjectsToPut = FromChangelist.ObjectsToPut;
+		const TMap<FObjectInStreamID, FConcertReplication_ChangeStream_PutObject>& ObjectsToPut = FromChangelist.ObjectsToPut;
 		
 		ConcertSyncClient::Replication::FChangeStreamRequest Request;
 		Request.StreamsToAdd.Emplace();
@@ -190,7 +190,7 @@ namespace UE::MultiUserClient
 		// If creating a new stream, the objects must be supplied in the description instead of in PutObjects!
 		FObjectReplicationMap& ReplicationMap = NewStream.BaseDescription.ReplicationMap;
 		ReplicationMap.ReplicatedObjects.Reserve(ObjectsToPut.Num());
-		for (const TPair<FObjectInStreamID, FConcertChangeStream_PutObject>& PutObjectPair : ObjectsToPut)
+		for (const TPair<FObjectInStreamID, FConcertReplication_ChangeStream_PutObject>& PutObjectPair : ObjectsToPut)
 		{
 			const TOptional<FReplicatedObjectInfo> NewObjectInfo = PutObjectPair.Value.MakeObjectInfoIfValid();
 			// The editing UI allows adding objects without properties (to make UX easier) - do not submit those to the server.
diff --git a/Engine/Plugins/Developer/Concert/ConcertApp/MultiUserClient/Source/MultiUserClient/Private/Replication/Stream/LocalClientStreamSynchronizer.h b/Engine/Plugins/Developer/Concert/ConcertApp/MultiUserClient/Source/MultiUserClient/Private/Replication/Stream/LocalClientStreamSynchronizer.h
--- a/Engine/Plugins/Developer/Concert/ConcertApp/MultiUserClient/Source/MultiUserClient/Private/Replication/Stream/LocalClientStreamSynchronizer.h
+++ b/Engine/Plugins/Developer/Concert/ConcertApp/MultiUserClient/Source/MultiUserClient/Private/Replication/Stream/LocalClientStreamSynchronizer.h
@@ -5,7 +5,7 @@
 #include "Replication/IToken.h"
 #include "Replication/Data/ObjectIds.h"
 #include "Replication/Data/ObjectReplicationMap.h"
-#include "Replication/Messages/ConcertReplicationEvents.h"
+#include "Replication/Messages/ChangeStream.h"
 
 #include "Async/Future.h"
 #include "Misc/Attribute.h"
@@ -130,7 +130,7 @@ namespace UE::MultiUserClient
 		struct FChangelist
 		{
 			TSet<FObjectInStreamID> ObjectsToRemove;
-			TMap<FObjectInStreamID, FConcertChangeStream_PutObject> ObjectsToPut;
+			TMap<FObjectInStreamID, FConcertReplication_ChangeStream_PutObject> ObjectsToPut;
 		};
 		
 		/**
diff --git a/Engine/Plugins/Developer/Concert/ConcertSync/ConcertSyncClient/Source/ConcertSyncClient/Private/Replication/Manager/ReplicationManager.cpp b/Engine/Plugins/Developer/Concert/ConcertSync/ConcertSyncClient/Source/ConcertSyncClient/Private/Replication/Manager/ReplicationManager.cpp
--- a/Engine/Plugins/Developer/Concert/ConcertSync/ConcertSyncClient/Source/ConcertSyncClient/Private/Replication/Manager/ReplicationManager.cpp
+++ b/Engine/Plugins/Developer/Concert/ConcertSync/ConcertSyncClient/Source/ConcertSyncClient/Private/Replication/Manager/ReplicationManager.cpp
@@ -6,7 +6,7 @@
 #include "ReplicationManagerState_Disconnected.h"
 #include "ReplicationManagerUtils.h"
 
-#include "Replication/Messages/ConcertReplicationHandshakeMessages.h"
+#include "Replication/Messages/Handshake.h"
 
 namespace UE::ConcertSyncClient::Replication
 {
diff --git a/Engine/Plugins/Developer/Concert/ConcertSync/ConcertSyncClient/Source/ConcertSyncClient/Private/Replication/Manager/ReplicationManager.h b/Engine/Plugins/Developer/Concert/ConcertSync/ConcertSyncClient/Source/ConcertSyncClient/Private/Replication/Manager/ReplicationManager.h
--- a/Engine/Plugins/Developer/Concert/ConcertSync/ConcertSyncClient/Source/ConcertSyncClient/Private/Replication/Manager/ReplicationManager.h
+++ b/Engine/Plugins/Developer/Concert/ConcertSync/ConcertSyncClient/Source/ConcertSyncClient/Private/Replication/Manager/ReplicationManager.h
@@ -3,7 +3,7 @@
 #pragma once
 
 #include "Replication/IConcertClientReplicationManager.h"
-#include "Replication/Messages/ConcertReplicationHandshakeMessages.h"
+#include "Replication/Messages/Handshake.h"
 #include "Templates/SharedPointer.h"
 
 class IConcertClientReplicationBridge;
diff --git a/Engine/Plugins/Developer/Concert/ConcertSync/ConcertSyncClient/Source/ConcertSyncClient/Private/Replication/Manager/ReplicationManagerState_Connected.cpp b/Engine/Plugins/Developer/Concert/ConcertSync/ConcertSyncClient/Source/ConcertSyncClient/Private/Replication/Manager/ReplicationManagerState_Connected.cpp
--- a/Engine/Plugins/Developer/Concert/ConcertSync/ConcertSyncClient/Source/ConcertSyncClient/Private/Replication/Manager/ReplicationManagerState_Connected.cpp
+++ b/Engine/Plugins/Developer/Concert/ConcertSync/ConcertSyncClient/Source/ConcertSyncClient/Private/Replication/Manager/ReplicationManagerState_Connected.cpp
@@ -8,7 +8,7 @@
 #include "ReplicationManagerState_Disconnected.h"
 #include "Replication/ChangeStreamSharedUtils.h"
 #include "Replication/Formats/FullObjectFormat.h"
-#include "Replication/Messages/ConcertReplicationHandshakeMessages.h"
+#include "Replication/Messages/Handshake.h"
 #include "Replication/Processing/ClientReplicationDataCollector.h"
 #include "Replication/Processing/ObjectReplicationApplierProcessor.h"
 #include "Replication/Processing/ObjectReplicationReceiver.h"
@@ -85,8 +85,8 @@ namespace UE::ConcertSyncClient::Replication
 		// At that point, it will log errors for receiving replication data from a client without authority.
 		HandleReleasingReplicatedObjects(Args);
 		
-		return LiveSession->SendCustomRequest<FConcertChangeAuthority_Request, FConcertChangeAuthority_Response>(Args, LiveSession->GetSessionServerEndpointId())
-			.Next([WeakThis = TWeakPtr<FReplicationManagerState_Connected>(SharedThis(this)), Args](FConcertChangeAuthority_Response&& Response) mutable
+		return LiveSession->SendCustomRequest<FConcertReplication_ChangeAuthority_Request, FConcertReplication_ChangeAuthority_Response>(Args, LiveSession->GetSessionServerEndpointId())
+			.Next([WeakThis = TWeakPtr<FReplicationManagerState_Connected>(SharedThis(this)), Args](FConcertReplication_ChangeAuthority_Response&& Response) mutable
 			{
 				if (const TSharedPtr<FReplicationManagerState_Connected> ThisPin = WeakThis.Pin())
 				{
@@ -105,8 +105,8 @@ namespace UE::ConcertSyncClient::Replication
 			return MakeFulfilledPromise<FClientQueryResponse>().GetFuture();
 		}
 		
-		return LiveSession->SendCustomRequest<FConcertQueryReplicationInfo_Request, FConcertQueryReplicationInfo_Response>(Args, LiveSession->GetSessionServerEndpointId())
-			.Next([](FConcertQueryReplicationInfo_Response&& Response)
+		return LiveSession->SendCustomRequest<FConcertReplication_QueryReplicationInfo_Request, FConcertReplication_QueryReplicationInfo_Response>(Args, LiveSession->GetSessionServerEndpointId())
+			.Next([](FConcertReplication_QueryReplicationInfo_Response&& Response)
 			{
 				return FClientQueryResponse { MoveTemp(Response) };
 			});
@@ -118,8 +118,8 @@ namespace UE::ConcertSyncClient::Replication
 		// At that point, it will log errors for receiving replication data from a client without authority.
 		HandleRemovingReplicatedObjects(Args);
 		
-		return LiveSession->SendCustomRequest<FConcertChangeStream_Request, FConcertChangeStream_Response>(Args, LiveSession->GetSessionServerEndpointId())
-			.Next([WeakThis = TWeakPtr<FReplicationManagerState_Connected>(SharedThis(this)), Args](FConcertChangeStream_Response&& Response)
+		return LiveSession->SendCustomRequest<FConcertReplication_ChangeStream_Request, FConcertReplication_ChangeStream_Response>(Args, LiveSession->GetSessionServerEndpointId())
+			.Next([WeakThis = TWeakPtr<FReplicationManagerState_Connected>(SharedThis(this)), Args](FConcertReplication_ChangeStream_Response&& Response)
 			{
 				const TSharedPtr<FReplicationManagerState_Connected> ThisPin = WeakThis.Pin();
 				if (ThisPin && Response.IsSuccess())
@@ -172,7 +172,7 @@ namespace UE::ConcertSyncClient::Replication
 	{
 		// Build RegisteredStreams while RegisteredStreams has the old, unupdated state
 		TMap<FSoftObjectPath, TArray<FGuid>> BundledModifiedObjects;
-		for (const TPair<FObjectInStreamID, FConcertChangeStream_PutObject>& PutObject : Request.ObjectsToPut)
+		for (const TPair<FObjectInStreamID, FConcertReplication_ChangeStream_PutObject>& PutObject : Request.ObjectsToPut)
 		{
 			const FObjectInStreamID ObjectInfo = PutObject.Key;
 			const FSoftObjectPath Object = ObjectInfo.Object;
@@ -214,7 +214,7 @@ namespace UE::ConcertSyncClient::Replication
 		}
 	}
 
-	void FReplicationManagerState_Connected::UpdateReplicatedObjectsAfterAuthorityChange(FAuthorityChangeRequest&& Request, const FConcertChangeAuthority_Response& Response) const
+	void FReplicationManagerState_Connected::UpdateReplicatedObjectsAfterAuthorityChange(FAuthorityChangeRequest&& Request, const FConcertReplication_ChangeAuthority_Response& Response) const
 	{
 		for (TPair<FSoftObjectPath, FConcertStreamArray>& TakeAuthority : Request.TakeAuthority)
 		{
diff --git a/Engine/Plugins/Developer/Concert/ConcertSync/ConcertSyncClient/Source/ConcertSyncClient/Private/Replication/Manager/ReplicationManagerState_Connected.h b/Engine/Plugins/Developer/Concert/ConcertSync/ConcertSyncClient/Source/ConcertSyncClient/Private/Replication/Manager/ReplicationManagerState_Connected.h
--- a/Engine/Plugins/Developer/Concert/ConcertSync/ConcertSyncClient/Source/ConcertSyncClient/Private/Replication/Manager/ReplicationManagerState_Connected.h
+++ b/Engine/Plugins/Developer/Concert/ConcertSync/ConcertSyncClient/Source/ConcertSyncClient/Private/Replication/Manager/ReplicationManagerState_Connected.h
@@ -126,7 +126,7 @@ namespace UE::ConcertSyncClient::Replication
 		 * @note Request is accepted as && because this function rewrites its memory when looking at rejections.
 		 * Since the request was already sent to the server it is assumed the request can just contain trash after.
 		 */
-		void UpdateReplicatedObjectsAfterAuthorityChange(FAuthorityChangeRequest&& Request, const FConcertChangeAuthority_Response& Response) const;
+		void UpdateReplicatedObjectsAfterAuthorityChange(FAuthorityChangeRequest&& Request, const FConcertReplication_ChangeAuthority_Response& Response) const;
 		void HandleReleasingReplicatedObjects(const FAuthorityChangeRequest& Request) const;
 	};
 }
diff --git a/Engine/Plugins/Developer/Concert/ConcertSync/ConcertSyncClient/Source/ConcertSyncClient/Private/Replication/Manager/ReplicationManagerState_Handshaking.cpp b/Engine/Plugins/Developer/Concert/ConcertSync/ConcertSyncClient/Source/ConcertSyncClient/Private/Replication/Manager/ReplicationManagerState_Handshaking.cpp
--- a/Engine/Plugins/Developer/Concert/ConcertSync/ConcertSyncClient/Source/ConcertSyncClient/Private/Replication/Manager/ReplicationManagerState_Handshaking.cpp
+++ b/Engine/Plugins/Developer/Concert/ConcertSync/ConcertSyncClient/Source/ConcertSyncClient/Private/Replication/Manager/ReplicationManagerState_Handshaking.cpp
@@ -5,7 +5,7 @@
 #include "IConcertSession.h"
 #include "ReplicationManagerState_Connected.h"
 #include "ReplicationManagerState_Disconnected.h"
-#include "Replication/Messages/ConcertReplicationHandshakeMessages.h"
+#include "Replication/Messages/Handshake.h"
 
 #include "Algo/Transform.h"
 
diff --git a/Engine/Plugins/Developer/Concert/ConcertSync/ConcertSyncClient/Source/ConcertSyncClient/Private/Replication/Manager/ReplicationManagerState_Handshaking.h b/Engine/Plugins/Developer/Concert/ConcertSync/ConcertSyncClient/Source/ConcertSyncClient/Private/Replication/Manager/ReplicationManagerState_Handshaking.h
--- a/Engine/Plugins/Developer/Concert/ConcertSync/ConcertSyncClient/Source/ConcertSyncClient/Private/Replication/Manager/ReplicationManagerState_Handshaking.h
+++ b/Engine/Plugins/Developer/Concert/ConcertSync/ConcertSyncClient/Source/ConcertSyncClient/Private/Replication/Manager/ReplicationManagerState_Handshaking.h
@@ -5,7 +5,7 @@
 #include "IConcertSessionHandler.h"
 #include "ReplicationManagerState.h"
 #include "ReplicationManagerUtils.h"
-#include "Replication/Messages/ConcertReplicationHandshakeMessages.h"
+#include "Replication/Messages/Handshake.h"
 #include "Replication/Processing/ClientReplicationDataCollector.h"
 
 class IConcertClientSession;
diff --git a/Engine/Plugins/Developer/Concert/ConcertSync/ConcertSyncClient/Source/ConcertSyncClient/Public/Replication/IConcertClientReplicationBridge.h b/Engine/Plugins/Developer/Concert/ConcertSync/ConcertSyncClient/Source/ConcertSyncClient/Public/Replication/IConcertClientReplicationBridge.h
--- a/Engine/Plugins/Developer/Concert/ConcertSync/ConcertSyncClient/Source/ConcertSyncClient/Public/Replication/IConcertClientReplicationBridge.h
+++ b/Engine/Plugins/Developer/Concert/ConcertSync/ConcertSyncClient/Source/ConcertSyncClient/Public/Replication/IConcertClientReplicationBridge.h
@@ -7,7 +7,7 @@
 #include "Delegates/DelegateCombinations.h"
 
 class UObject;
-struct FConcertObjectReplicationEvent;
+struct FConcertReplication_ObjectReplicationEvent;
 struct FSoftObjectPath;
 
 DECLARE_MULTICAST_DELEGATE_OneParam(FConcertClientReplicationBridgeObjectEvent, UObject&);
diff --git a/Engine/Plugins/Developer/Concert/ConcertSync/ConcertSyncClient/Source/ConcertSyncClient/Public/Replication/IConcertClientReplicationManager.h b/Engine/Plugins/Developer/Concert/ConcertSync/ConcertSyncClient/Source/ConcertSyncClient/Public/Replication/IConcertClientReplicationManager.h
--- a/Engine/Plugins/Developer/Concert/ConcertSync/ConcertSyncClient/Source/ConcertSyncClient/Public/Replication/IConcertClientReplicationManager.h
+++ b/Engine/Plugins/Developer/Concert/ConcertSync/ConcertSyncClient/Source/ConcertSyncClient/Public/Replication/IConcertClientReplicationManager.h
@@ -3,9 +3,11 @@
 #pragma once
 
 #include "Misc/EBreakBehavior.h"
-#include "Replication/Messages/ConcertReplicationHandshakeMessages.h"
+#include "Replication/Messages/Handshake.h"
 #include "Replication/Data/ReplicationClientDescription.h"
-#include "Replication/Messages/ConcertReplicationEvents.h"
+#include "Replication/Messages/ChangeAuthority.h"
+#include "Replication/Messages/ChangeStream.h"
+#include "Replication/Messages/ClientQuery.h"
 
 template<typename ResultType>
 class TFuture;
@@ -37,12 +39,12 @@ namespace UE::ConcertSyncClient::Replication
 	};
 
 	// The intention here is to wrap the request in case there is some more specific meta data we want to add in the future.
-	struct FAuthorityChangeRequest : FConcertChangeAuthority_Request {};
-	struct FAuthorityChangeResponse : FConcertChangeAuthority_Response {};
-	struct FClientQueryRequest : FConcertQueryReplicationInfo_Request {};
-	struct FClientQueryResponse : FConcertQueryReplicationInfo_Response {};
-	struct FChangeStreamRequest : FConcertChangeStream_Request {};
-	struct FChangeStreamResponse : FConcertChangeStream_Response {};
+	struct FAuthorityChangeRequest : FConcertReplication_ChangeAuthority_Request {};
+	struct FAuthorityChangeResponse : FConcertReplication_ChangeAuthority_Response {};
+	struct FClientQueryRequest : FConcertReplication_QueryReplicationInfo_Request {};
+	struct FClientQueryResponse : FConcertReplication_QueryReplicationInfo_Response {};
+	struct FChangeStreamRequest : FConcertReplication_ChangeStream_Request {};
+	struct FChangeStreamResponse : FConcertReplication_ChangeStream_Response {};
 }
 
 /**
diff --git a/Engine/Plugins/Developer/Concert/ConcertSync/ConcertSyncCore/Source/ConcertSyncCore/Private/Replication/ChangeStreamSharedUtils.cpp b/Engine/Plugins/Developer/Concert/ConcertSync/ConcertSyncCore/Source/ConcertSyncCore/Private/Replication/ChangeStreamSharedUtils.cpp
--- a/Engine/Plugins/Developer/Concert/ConcertSync/ConcertSyncCore/Source/ConcertSyncCore/Private/Replication/ChangeStreamSharedUtils.cpp
+++ b/Engine/Plugins/Developer/Concert/ConcertSync/ConcertSyncCore/Source/ConcertSyncCore/Private/Replication/ChangeStreamSharedUtils.cpp
@@ -4,12 +4,13 @@
 
 #include "Misc/EBreakBehavior.h"
 #include "Replication/Data/ObjectIds.h"
-#include "Replication/Messages/ConcertReplicationEvents.h"
+#include "Replication/Messages/ChangeStream.h"
+#include "Replication/Messages/ObjectReplication.h"
 
 namespace UE::ConcertSyncCore::Replication::ChangeStreamUtils
 {
 	void ForEachObjectLosingAuthority(
-		const FConcertChangeStream_Request& Request,
+		const FConcertReplication_ChangeStream_Request& Request,
 		const TArray<FReplicationStreamDescription>& ExistingStreams,
 		TFunctionRef<EBreakBehavior(const FObjectInStreamID&)> Callback
 		)
@@ -41,7 +42,7 @@ namespace UE::ConcertSyncCore::Replication::ChangeStreamUtils
 		}
 	}
 	
-	void ApplyValidatedRequest(const FConcertChangeStream_Request& Request, TArray<FReplicationStreamDescription>& StreamsToModify)
+	void ApplyValidatedRequest(const FConcertReplication_ChangeStream_Request& Request, TArray<FReplicationStreamDescription>& StreamsToModify)
 	{
 		for (auto StreamIt = StreamsToModify.CreateIterator(); StreamIt; ++StreamIt)
 		{
@@ -64,7 +65,7 @@ namespace UE::ConcertSyncCore::Replication::ChangeStreamUtils
 				StreamIt.RemoveCurrent();
 			}
 			
-			for (const TPair<FObjectInStreamID, FConcertChangeStream_PutObject>& PutObjectPair : Request.ObjectsToPut)
+			for (const TPair<FObjectInStreamID, FConcertReplication_ChangeStream_PutObject>& PutObjectPair : Request.ObjectsToPut)
 			{
 				if (PutObjectPair.Key.StreamId != StreamId)
 				{
@@ -119,7 +120,7 @@ namespace UE::ConcertSyncCore::Replication::ChangeStreamUtils
 
 	namespace Private
 	{
-		static void BuildPutObjectList(const FGuid& StreamId, const FObjectReplicationMap& Base, const FObjectReplicationMap& Desired, FConcertChangeStream_Request& Request)
+		static void BuildPutObjectList(const FGuid& StreamId, const FObjectReplicationMap& Base, const FObjectReplicationMap& Desired, FConcertReplication_ChangeStream_Request& Request)
 		{
 			for (const TPair<FSoftObjectPath, FReplicatedObjectInfo>& BasePair : Base.ReplicatedObjects)
 			{
@@ -129,7 +130,7 @@ namespace UE::ConcertSyncCore::Replication::ChangeStreamUtils
 				if (DesiredObjectInfo)
 				{
 					const FReplicatedObjectInfo& BaseObjectInfo = BasePair.Value;
-					const TOptional<FConcertChangeStream_PutObject> PutObject = FConcertChangeStream_PutObject::MakeFromChange(BaseObjectInfo, *DesiredObjectInfo);
+					const TOptional<FConcertReplication_ChangeStream_PutObject> PutObject = FConcertReplication_ChangeStream_PutObject::MakeFromChange(BaseObjectInfo, *DesiredObjectInfo);
 				
 					const bool bDesiredHasChangedFromBase = BaseObjectInfo != *DesiredObjectInfo;
 					// If MakeFromChange returned unset, it means that this request is not valid to submit. 
@@ -146,7 +147,7 @@ namespace UE::ConcertSyncCore::Replication::ChangeStreamUtils
 			}
 		}
 
-		static void BuildRemoveObjectList(const FGuid& StreamId, const FObjectReplicationMap& Base, const FObjectReplicationMap& Desired, FConcertChangeStream_Request& Request)
+		static void BuildRemoveObjectList(const FGuid& StreamId, const FObjectReplicationMap& Base, const FObjectReplicationMap& Desired, FConcertReplication_ChangeStream_Request& Request)
 		{
 			for (const TPair<FSoftObjectPath, FReplicatedObjectInfo>& DesiredPair : Desired.ReplicatedObjects)
 			{
@@ -159,7 +160,7 @@ namespace UE::ConcertSyncCore::Replication::ChangeStreamUtils
 
 				// Desired wants to add an object
 				const FReplicatedObjectInfo& ObjectInfo = DesiredPair.Value;
-				const TOptional<FConcertChangeStream_PutObject> PutObject = FConcertChangeStream_PutObject::MakeFromInfo(ObjectInfo);
+				const TOptional<FConcertReplication_ChangeStream_PutObject> PutObject = FConcertReplication_ChangeStream_PutObject::MakeFromInfo(ObjectInfo);
 				const bool bDesiredStateHasEnoughDataToForPut = ensureMsgf(PutObject, TEXT("Function assumption violated; you did not pass in valid base or desired state."));
 				if (bDesiredStateHasEnoughDataToForPut)
 				{
@@ -169,13 +170,13 @@ namespace UE::ConcertSyncCore::Replication::ChangeStreamUtils
 		}
 	}
 
-	FConcertChangeStream_Request BuildRequestFromDiff(
+	FConcertReplication_ChangeStream_Request BuildRequestFromDiff(
 		const FGuid& StreamId,
 		const FObjectReplicationMap& Base,
 		const FObjectReplicationMap& Desired
 		)
 	{
-		FConcertChangeStream_Request Request;
+		FConcertReplication_ChangeStream_Request Request;
 		Private::BuildPutObjectList(StreamId, Base, Desired, Request);
 		Private::BuildRemoveObjectList(StreamId, Base, Desired, Request);
 		return Request;
diff --git a/Engine/Plugins/Developer/Concert/ConcertSync/ConcertSyncCore/Source/ConcertSyncCore/Private/Replication/Messages/ChangeStream.cpp b/Engine/Plugins/Developer/Concert/ConcertSync/ConcertSyncCore/Source/ConcertSyncCore/Private/Replication/Messages/ChangeStream.cpp
new file mode 100644
--- /dev/null
+++ b/Engine/Plugins/Developer/Concert/ConcertSync/ConcertSyncCore/Source/ConcertSyncCore/Private/Replication/Messages/ChangeStream.cpp
@@ -0,0 +1,61 @@
+// Copyright Epic Games, Inc. All Rights Reserved.
+
+#include "Replication/Messages/ChangeStream.h"
+
+#include "Misc/OutputDevice.h"
+
+TOptional<FConcertReplication_ChangeStream_PutObject> FConcertReplication_ChangeStream_PutObject::MakeFromInfo(const FReplicatedObjectInfo& New)
+{
+	if (!New.IsValidForSendingToServer())
+	{
+		return {};
+	}
+	return FConcertReplication_ChangeStream_PutObject{ New.PropertySelection, New.ClassPath };
+}
+
+TOptional<FConcertReplication_ChangeStream_PutObject> FConcertReplication_ChangeStream_PutObject::MakeFromChange(const FReplicatedObjectInfo& Base, const FReplicatedObjectInfo& Desired)
+{
+	if (Base.IsValidForSendingToServer() && Desired.IsValidForSendingToServer())
+	{
+		const FConcertPropertySelection PropertySelection = Base.PropertySelection != Desired.PropertySelection ? Desired.PropertySelection : FConcertPropertySelection{};
+		const FSoftClassPath ClassPath = Base.ClassPath != Desired.ClassPath ? Desired.ClassPath : FSoftClassPath{};
+		return FConcertReplication_ChangeStream_PutObject{ PropertySelection, ClassPath };
+	}
+	return {};
+}
+
+TOptional<FReplicatedObjectInfo> FConcertReplication_ChangeStream_PutObject::MakeObjectInfoIfValid() const
+{
+	if (Properties.ReplicatedProperties.IsEmpty() || ClassPath.IsNull())
+	{
+		return {};
+	}
+	return FReplicatedObjectInfo{ ClassPath, Properties };
+}
+
+void FConcertReplication_ChangeStream_Response::LogErrors(FOutputDevice& OutputDevice) const
+{
+	for (const TPair<FObjectInStreamID, FReplicatedObjectId>& Conflict : AuthorityConflicts)
+	{
+		OutputDevice.Logf(TEXT("Authority: { %s } conflicts with { %s }"), *Conflict.Key.ToString(), *Conflict.Value.ToString());
+	}
+
+	for (const TPair<FObjectInStreamID, EConcertPutObjectErrorCode>& Error : ObjectsToPutSemanticErrors)
+	{
+		const FString ErrorCodeAsString = [&Error]()
+		{
+			switch (Error.Value)
+			{
+			case EConcertPutObjectErrorCode::UnresolvedStream: return TEXT("Unresolved stream");
+			case EConcertPutObjectErrorCode::MissingData: return TEXT("Missing data");
+			default: return TEXT("Unknown");
+			}
+		}();
+		OutputDevice.Logf(TEXT("Semantic error: %s - %s"), *Error.Key.ToString(), *ErrorCodeAsString);
+	}
+
+	for (const FGuid& Stream : FailedStreamCreation)
+	{
+		OutputDevice.Logf(TEXT("Failed to create stream: %s"), *Stream.ToString());
+	}
+}
diff --git a/Engine/Plugins/Developer/Concert/ConcertSync/ConcertSyncCore/Source/ConcertSyncCore/Private/Replication/Messages/ConcertReplicationEvents.cpp b/Engine/Plugins/Developer/Concert/ConcertSync/ConcertSyncCore/Source/ConcertSyncCore/Private/Replication/Messages/ConcertReplicationEvents.cpp
deleted file mode 100644
diff --git a/Engine/Plugins/Developer/Concert/ConcertSync/ConcertSyncCore/Source/ConcertSyncCore/Private/Replication/Messages/ConcertReplicationHandshakeMessages.cpp b/Engine/Plugins/Developer/Concert/ConcertSync/ConcertSyncCore/Source/ConcertSyncCore/Private/Replication/Messages/ConcertReplicationHandshakeMessages.cpp
deleted file mode 100644
diff --git a/Engine/Plugins/Developer/Concert/ConcertSync/ConcertSyncCore/Source/ConcertSyncCore/Private/Replication/Messages/Handshake.cpp b/Engine/Plugins/Developer/Concert/ConcertSync/ConcertSyncCore/Source/ConcertSyncCore/Private/Replication/Messages/Handshake.cpp
new file mode 100644
--- /dev/null
+++ b/Engine/Plugins/Developer/Concert/ConcertSync/ConcertSyncCore/Source/ConcertSyncCore/Private/Replication/Messages/Handshake.cpp
@@ -0,0 +1,27 @@
+// Copyright Epic Games, Inc. All Rights Reserved.
+
+#include "Replication/Messages/Handshake.h"
+
+namespace UE::ConcertSyncCore::Replication
+{
+	FString LexJoinErrorCode(EJoinReplicationErrorCode ErrorCode)
+	{
+		static_assert(static_cast<int32>(EJoinReplicationErrorCode::Max) == 9, "Update JoinErrorCodeToText when you change entries");
+
+		switch (ErrorCode)
+		{
+		case EJoinReplicationErrorCode::Success: return TEXT("Success");
+		case EJoinReplicationErrorCode::NetworkError: return TEXT("NetworkError");
+		case EJoinReplicationErrorCode::Cancelled: return TEXT("Cancelled");
+		case EJoinReplicationErrorCode::AlreadyInProgress: return TEXT("AlreadyInProgress");
+		case EJoinReplicationErrorCode::NotInAnyConcertSession: return TEXT("NotInAnyConcertSession");
+		case EJoinReplicationErrorCode::InvalidClass: return TEXT("InvalidClass");
+		case EJoinReplicationErrorCode::DuplicateProperty: return TEXT("DuplicateProperty");
+		case EJoinReplicationErrorCode::AlreadyInSession: return TEXT("AlreadyInSession");
+		case EJoinReplicationErrorCode::DuplicateStreamId: return TEXT("DuplicateStreamId");
+		case EJoinReplicationErrorCode::FailedToUnpackStream: return TEXT("FailedToUnpackStream");
+		default: 
+			return FString();
+		}
+	}
+}
\ No newline at end of file
diff --git a/Engine/Plugins/Developer/Concert/ConcertSync/ConcertSyncCore/Source/ConcertSyncCore/Private/Replication/Processing/ObjectReplicationCache.cpp b/Engine/Plugins/Developer/Concert/ConcertSync/ConcertSyncCore/Source/ConcertSyncCore/Private/Replication/Processing/ObjectReplicationCache.cpp
--- a/Engine/Plugins/Developer/Concert/ConcertSync/ConcertSyncCore/Source/ConcertSyncCore/Private/Replication/Processing/ObjectReplicationCache.cpp
+++ b/Engine/Plugins/Developer/Concert/ConcertSync/ConcertSyncCore/Source/ConcertSyncCore/Private/Replication/Processing/ObjectReplicationCache.cpp
@@ -4,7 +4,7 @@
 
 #include "Replication/Data/ObjectIds.h"
 #include "Replication/Formats/IObjectReplicationFormat.h"
-#include "Replication/Messages/ConcertReplicationEvents.h"
+#include "Replication/Messages/ObjectReplication.h"
 
 namespace UE::ConcertSyncCore
 {
@@ -12,7 +12,7 @@ namespace UE::ConcertSyncCore
 		: ReplicationFormat(MoveTemp(ReplicationFormat))
 	{}
 
-	int32 FObjectReplicationCache::StoreUntilConsumed(const FGuid& SendingEndpointId, const FGuid& OriginStreamId, const FConcertObjectReplicationEvent& ObjectReplicationEvent)
+	int32 FObjectReplicationCache::StoreUntilConsumed(const FGuid& SendingEndpointId, const FGuid& OriginStreamId, const FConcertReplication_ObjectReplicationEvent& ObjectReplicationEvent)
 	{
 		int32 NumAccepted = 0;
 		
@@ -21,11 +21,11 @@ namespace UE::ConcertSyncCore
 		if (ObjectCacheBeforeAddition)
 		{
 			// It is important to compare by address and not by TWeakPtr instance because each user has a different TWeakPtr instance (but they all point to the same memory).
-			TSet<FConcertObjectReplicationEvent*> CombineOnceDetection;
-			for (const TPair<TWeakPtr<IReplicationCacheUser>, TWeakPtr<FConcertObjectReplicationEvent>>& InUseDataPair : ObjectCacheBeforeAddition->DataInUse)
+			TSet<FConcertReplication_ObjectReplicationEvent*> CombineOnceDetection;
+			for (const TPair<TWeakPtr<IReplicationCacheUser>, TWeakPtr<FConcertReplication_ObjectReplicationEvent>>& InUseDataPair : ObjectCacheBeforeAddition->DataInUse)
 			{
-				const TWeakPtr<FConcertObjectReplicationEvent>& EventData = InUseDataPair.Value;
-				const TSharedPtr<FConcertObjectReplicationEvent> EventDataPin = EventData.Pin();
+				const TWeakPtr<FConcertReplication_ObjectReplicationEvent>& EventData = InUseDataPair.Value;
+				const TSharedPtr<FConcertReplication_ObjectReplicationEvent> EventDataPin = EventData.Pin();
 				if (EventDataPin && !CombineOnceDetection.Contains(EventDataPin.Get()))
 				{
 					CombineOnceDetection.Add(EventDataPin.Get());
@@ -35,7 +35,7 @@ namespace UE::ConcertSyncCore
 		}
 
 		FObjectCache* ObjectCacheAfterAddition = ObjectCacheBeforeAddition ? ObjectCacheBeforeAddition : nullptr;
-		TSharedPtr<FConcertObjectReplicationEvent> LazilyCopiedEventPtr;
+		TSharedPtr<FConcertReplication_ObjectReplicationEvent> LazilyCopiedEventPtr;
 		for (const TSharedRef<IReplicationCacheUser>& CacheUser : CacheUsers)
 		{
 			if (ObjectCacheBeforeAddition && ObjectCacheBeforeAddition->DataInUse.Contains(CacheUser))
@@ -50,13 +50,13 @@ namespace UE::ConcertSyncCore
 				// Do the event copy only when somebody wants the data...
 				if (!LazilyCopiedEventPtr.IsValid())
 				{
-					LazilyCopiedEventPtr = MakeShared<FConcertObjectReplicationEvent>(ObjectReplicationEvent);
+					LazilyCopiedEventPtr = MakeShared<FConcertReplication_ObjectReplicationEvent>(ObjectReplicationEvent);
 				}
 
 				// We'll constructor a new proxy shared ptr which will clean-up Cache when released
 				TWeakPtr<IReplicationCacheUser> WeakUserPtr = CacheUser;
 				TWeakPtr<FObjectReplicationCache> WeakThisPtr = AsWeak();
-				TSharedRef<FConcertObjectReplicationEvent> GuardPtr = MakeShareable(LazilyCopiedEventPtr.Get(), [OriginStreamId, LazilyCopiedEventPtr, WeakUserPtr, WeakThisPtr](auto*)
+				TSharedRef<FConcertReplication_ObjectReplicationEvent> GuardPtr = MakeShareable(LazilyCopiedEventPtr.Get(), [OriginStreamId, LazilyCopiedEventPtr, WeakUserPtr, WeakThisPtr](auto*)
 				{
 					// Replication cache user can survive the destruction of the cache because it is created externally
 					const TSharedPtr<FObjectReplicationCache> This = WeakThisPtr.Pin();
@@ -82,7 +82,7 @@ namespace UE::ConcertSyncCore
 					// LazyCopiedEventPtr will decrement counter and possibly be destroyed now
 				});
 				
-				const TWeakPtr<FConcertObjectReplicationEvent> WeakGuardPtr = GuardPtr;
+				const TWeakPtr<FConcertReplication_ObjectReplicationEvent> WeakGuardPtr = GuardPtr;
 				CacheUser->OnDataCached(ObjectId, MoveTemp(GuardPtr));
 				// Was it instantly consumed? Should not really happen but it could technically...
 				if (!LIKELY(WeakGuardPtr.IsValid()))
diff --git a/Engine/Plugins/Developer/Concert/ConcertSync/ConcertSyncCore/Source/ConcertSyncCore/Private/Replication/Processing/ObjectReplicationReceiver.cpp b/Engine/Plugins/Developer/Concert/ConcertSync/ConcertSyncCore/Source/ConcertSyncCore/Private/Replication/Processing/ObjectReplicationReceiver.cpp
--- a/Engine/Plugins/Developer/Concert/ConcertSync/ConcertSyncCore/Source/ConcertSyncCore/Private/Replication/Processing/ObjectReplicationReceiver.cpp
+++ b/Engine/Plugins/Developer/Concert/ConcertSync/ConcertSyncCore/Source/ConcertSyncCore/Private/Replication/Processing/ObjectReplicationReceiver.cpp
@@ -4,7 +4,7 @@
 
 #include "ConcertLogGlobal.h"
 #include "IConcertSession.h"
-#include "Replication/Messages/ConcertReplicationEvents.h"
+#include "Replication/Messages/ObjectReplication.h"
 #include "Replication/Processing/ObjectReplicationCache.h"
 
 #include "HAL/IConsoleManager.h"
@@ -17,15 +17,15 @@ namespace UE::ConcertSyncCore
 		: Session(MoveTemp(Session))
 		, ReplicationCache(MoveTemp(ReplicationCache))
 	{
-		Session->RegisterCustomEventHandler<FConcertBatchReplicationEvent>(this, &FObjectReplicationReceiver::HandleBatchReplicationEvent);
+		Session->RegisterCustomEventHandler<FConcertReplication_BatchReplicationEvent>(this, &FObjectReplicationReceiver::HandleBatchReplicationEvent);
 	}
 
 	FObjectReplicationReceiver::~FObjectReplicationReceiver()
 	{
-		Session->UnregisterCustomEventHandler<FConcertBatchReplicationEvent>(this);
+		Session->UnregisterCustomEventHandler<FConcertReplication_BatchReplicationEvent>(this);
 	}
 
-	void FObjectReplicationReceiver::HandleBatchReplicationEvent(const FConcertSessionContext& SessionContext, const FConcertBatchReplicationEvent& Event)
+	void FObjectReplicationReceiver::HandleBatchReplicationEvent(const FConcertSessionContext& SessionContext, const FConcertReplication_BatchReplicationEvent& Event)
 	{
 		// Fyi: an object may have multiple changes in a batch replication event: each stream can modify different properties as long as they do not overlap.
 		int32 NumObjects = 0;
@@ -33,12 +33,12 @@ namespace UE::ConcertSyncCore
 		int32 NumCacheUsages = 0;
 		int32 NumOfAcceptedObjectChanges = 0;
 		
-		for (const FConcertStreamReplicationEvent& StreamEvent : Event.Streams)
+		for (const FConcertReplication_StreamReplicationEvent& StreamEvent : Event.Streams)
 		{
 			const int32 ObjectsInStream = StreamEvent.ReplicatedObjects.Num();
 			NumObjects += ObjectsInStream;
 			
-			for (const FConcertObjectReplicationEvent& ObjectEvent : StreamEvent.ReplicatedObjects)
+			for (const FConcertReplication_ObjectReplicationEvent& ObjectEvent : StreamEvent.ReplicatedObjects)
 			{
 				if (ShouldAcceptObject(SessionContext, StreamEvent, ObjectEvent))
 				{
diff --git a/Engine/Plugins/Developer/Concert/ConcertSync/ConcertSyncCore/Source/ConcertSyncCore/Private/Replication/Processing/ObjectReplicationSender.cpp b/Engine/Plugins/Developer/Concert/ConcertSync/ConcertSyncCore/Source/ConcertSyncCore/Private/Replication/Processing/ObjectReplicationSender.cpp
--- a/Engine/Plugins/Developer/Concert/ConcertSync/ConcertSyncCore/Source/ConcertSyncCore/Private/Replication/Processing/ObjectReplicationSender.cpp
+++ b/Engine/Plugins/Developer/Concert/ConcertSync/ConcertSyncCore/Source/ConcertSyncCore/Private/Replication/Processing/ObjectReplicationSender.cpp
@@ -29,7 +29,7 @@ namespace UE::ConcertSyncCore
 
 		if (!EventToSend.Streams.IsEmpty())
 		{
-			const int32 NumObjects = Algo::TransformAccumulate(EventToSend.Streams, [](const FConcertStreamReplicationEvent& Event){ return Event.ReplicatedObjects.Num(); }, 0);
+			const int32 NumObjects = Algo::TransformAccumulate(EventToSend.Streams, [](const FConcertReplication_StreamReplicationEvent& Event){ return Event.ReplicatedObjects.Num(); }, 0);
 			UE_CLOG(CVarLogSentObjects.GetValueOnGameThread(), LogConcert, Log, TEXT("Sending %d streams with %d objects to %s"),
 				EventToSend.Streams.Num(),
 				NumObjects,
@@ -53,8 +53,8 @@ namespace UE::ConcertSyncCore
 		const FSoftObjectPath& ReplicatedObject = Args.ObjectInfo.Object;
 		auto CaptureData = [this, &Args, &ReplicatedObject]<typename TPayloadRefType>(TPayloadRefType&& Payload)
 		{
-			const int32 PreexistingIndex = EventToSend.Streams.IndexOfByPredicate([&Args](const FConcertStreamReplicationEvent& StreamData){ return StreamData.StreamId == Args.ObjectInfo.StreamId; });
-			FConcertStreamReplicationEvent& StreamData = EventToSend.Streams.IsValidIndex(PreexistingIndex)
+			const int32 PreexistingIndex = EventToSend.Streams.IndexOfByPredicate([&Args](const FConcertReplication_StreamReplicationEvent& StreamData){ return StreamData.StreamId == Args.ObjectInfo.StreamId; });
+			FConcertReplication_StreamReplicationEvent& StreamData = EventToSend.Streams.IsValidIndex(PreexistingIndex)
 				? EventToSend.Streams[PreexistingIndex]
 				: EventToSend.Streams[EventToSend.Streams.Emplace(Args.ObjectInfo.StreamId)];
 			StreamData.ReplicatedObjects.Add({
diff --git a/Engine/Plugins/Developer/Concert/ConcertSync/ConcertSyncCore/Source/ConcertSyncCore/Private/Replication/Processing/ReplicationDataQueuer.cpp b/Engine/Plugins/Developer/Concert/ConcertSync/ConcertSyncCore/Source/ConcertSyncCore/Private/Replication/Processing/ReplicationDataQueuer.cpp
--- a/Engine/Plugins/Developer/Concert/ConcertSync/ConcertSyncCore/Source/ConcertSyncCore/Private/Replication/Processing/ReplicationDataQueuer.cpp
+++ b/Engine/Plugins/Developer/Concert/ConcertSync/ConcertSyncCore/Source/ConcertSyncCore/Private/Replication/Processing/ReplicationDataQueuer.cpp
@@ -2,7 +2,7 @@
 
 #include "Replication/Processing/ReplicationDataQueuer.h"
 
-#include "Replication/Messages/ConcertReplicationEvents.h"
+#include "Replication/Messages/ObjectReplication.h"
 
 namespace UE::ConcertSyncCore
 {
@@ -33,7 +33,7 @@ namespace UE::ConcertSyncCore
 		TFunctionRef<void(FConcertSessionSerializedPayload&& Payload)> ProcessMoveable
 		)
 	{
-		TSharedPtr<const FConcertObjectReplicationEvent> Event;
+		TSharedPtr<const FConcertReplication_ObjectReplicationEvent> Event;
 		const bool bSuccess = PendingEvents.RemoveAndCopyValue(Object, Event);
 		if (!ensureMsgf(bSuccess, TEXT("ExtractReplicationDataForObject for an item that was not returned by ForEachPendingObject")))
 		{
@@ -45,7 +45,7 @@ namespace UE::ConcertSyncCore
 		return true;
 	}
 
-	void FReplicationDataQueuer::OnDataCached(const FReplicatedObjectId& Object, TSharedRef<const FConcertObjectReplicationEvent> Data)
+	void FReplicationDataQueuer::OnDataCached(const FReplicatedObjectId& Object, TSharedRef<const FConcertReplication_ObjectReplicationEvent> Data)
 	{
 		PendingEvents.Add(Object, MoveTemp(Data));
 	}
diff --git a/Engine/Plugins/Developer/Concert/ConcertSync/ConcertSyncCore/Source/ConcertSyncCore/Public/Replication/ChangeStreamSharedUtils.h b/Engine/Plugins/Developer/Concert/ConcertSync/ConcertSyncCore/Source/ConcertSyncCore/Public/Replication/ChangeStreamSharedUtils.h
--- a/Engine/Plugins/Developer/Concert/ConcertSync/ConcertSyncCore/Source/ConcertSyncCore/Public/Replication/ChangeStreamSharedUtils.h
+++ b/Engine/Plugins/Developer/Concert/ConcertSync/ConcertSyncCore/Source/ConcertSyncCore/Public/Replication/ChangeStreamSharedUtils.h
@@ -8,7 +8,7 @@
 
 enum class EBreakBehavior : uint8;
 
-struct FConcertChangeStream_Request;
+struct FConcertReplication_ChangeStream_Request;
 struct FObjectInStreamID;
 struct FObjectReplicationMap;
 struct FReplicationStreamDescription;
@@ -25,7 +25,7 @@ namespace UE::ConcertSyncCore::Replication::ChangeStreamUtils
 	 * @param Callback The callback to process the objects
 	 */
 	CONCERTSYNCCORE_API void ForEachObjectLosingAuthority(
-		const FConcertChangeStream_Request& Request,
+		const FConcertReplication_ChangeStream_Request& Request,
 		const TArray<FReplicationStreamDescription>& ExistingStreams,
 		TFunctionRef<EBreakBehavior(const FObjectInStreamID&)> Callback
 		);
@@ -36,7 +36,7 @@ namespace UE::ConcertSyncCore::Replication::ChangeStreamUtils
 	 * @param StreamsToModify The streams to apply the request to
 	 */
 	CONCERTSYNCCORE_API void ApplyValidatedRequest(
-		const FConcertChangeStream_Request& Request,
+		const FConcertReplication_ChangeStream_Request& Request,
 		IN OUT TArray<FReplicationStreamDescription>& StreamsToModify
 		);
 
@@ -56,7 +56,7 @@ namespace UE::ConcertSyncCore::Replication::ChangeStreamUtils
 	 * @param Desired The state the stream should end up in 
 	 * @return A semantically correct FConcertChangeStream_Request which Base to Desired (note is can be rejected due to authority conflicts).
 	 */
-	CONCERTSYNCCORE_API FConcertChangeStream_Request BuildRequestFromDiff(
+	CONCERTSYNCCORE_API FConcertReplication_ChangeStream_Request BuildRequestFromDiff(
 		const FGuid& StreamId,
 		const FObjectReplicationMap& Base,
 		const FObjectReplicationMap& Desired
diff --git a/Engine/Plugins/Developer/Concert/ConcertSync/ConcertSyncCore/Source/ConcertSyncCore/Public/Replication/Messages/ChangeAuthority.h b/Engine/Plugins/Developer/Concert/ConcertSync/ConcertSyncCore/Source/ConcertSyncCore/Public/Replication/Messages/ChangeAuthority.h
new file mode 100644
--- /dev/null
+++ b/Engine/Plugins/Developer/Concert/ConcertSync/ConcertSyncCore/Source/ConcertSyncCore/Public/Replication/Messages/ChangeAuthority.h
@@ -0,0 +1,60 @@
+// Copyright Epic Games, Inc. All Rights Reserved.
+
+#pragma once
+
+#include "Replication/Data/ObjectIds.h"
+#include "ChangeAuthority.generated.h"
+
+USTRUCT()
+struct FConcertStreamArray
+{
+	GENERATED_BODY()
+
+	UPROPERTY()
+	TArray<FGuid> StreamIds;
+};
+
+/**
+ * Sent by clients to the server to attempt to take or release authority over properties on objects.
+ * Clients must have authority over object properties before sending FConcertBatchReplicationEvent, or the server will reject the updates.
+ * 
+ * Multiple clients can have authority over the same objects as long as the properties do not overlap.
+ * The properties are defined implicitly by providing the stream ID that will be writing the properties.
+ */
+USTRUCT()
+struct FConcertReplication_ChangeAuthority_Request
+{
+	GENERATED_BODY()
+
+	// In the future we can consider adding an option to make this all or nothing (if one object fails, the entire request does)
+	// It is supposed to be a convenience and network optimization to send one big request instead of many small ones.
+
+	/**
+	 * Objects the client requests to take authority over.
+	 * Mapping of object to stream identifiers the client has previously registered.
+	 */
+	UPROPERTY()
+	TMap<FSoftObjectPath, FConcertStreamArray> TakeAuthority;
+
+	/**
+	 * Objects the client no longer needs authority over.
+	 * Mapping of object to stream identifiers the client has previously registered.
+	 * 
+	 * The request will not fail if it contains objects it has no authority over.
+	 */
+	UPROPERTY()
+	TMap<FSoftObjectPath, FConcertStreamArray> ReleaseAuthority;
+};
+
+USTRUCT()
+struct FConcertReplication_ChangeAuthority_Response
+{
+	GENERATED_BODY()
+
+	/**
+	 * Objects streams the client did not receive authority over.
+	 * The client is implied to have authority over all other request objects.
+	 */
+	UPROPERTY()
+	TMap<FSoftObjectPath, FConcertStreamArray> RejectedObjects;
+};
\ No newline at end of file
diff --git a/Engine/Plugins/Developer/Concert/ConcertSync/ConcertSyncCore/Source/ConcertSyncCore/Public/Replication/Messages/ChangeStream.h b/Engine/Plugins/Developer/Concert/ConcertSync/ConcertSyncCore/Source/ConcertSyncCore/Public/Replication/Messages/ChangeStream.h
new file mode 100644
--- /dev/null
+++ b/Engine/Plugins/Developer/Concert/ConcertSync/ConcertSyncCore/Source/ConcertSyncCore/Public/Replication/Messages/ChangeStream.h
@@ -0,0 +1,165 @@
+// Copyright Epic Games, Inc. All Rights Reserved.
+
+#pragma once
+
+#include "Replication/Data/ObjectIds.h"
+#include "Replication/Data/ConcertPropertySelection.h"
+#include "Replication/Data/ReplicationStreamDescription.h"
+
+#include "Misc/Optional.h"
+#include "ChangeStream.generated.h"
+
+class FOutputDevice;
+
+struct FReplicatedObjectInfo;
+
+/** A request to add a new object to a stream or overwrite a pre-existing object's properties / class. */
+USTRUCT()
+struct FConcertReplication_ChangeStream_PutObject
+{
+	GENERATED_BODY()
+
+	/**
+	 * The property selection the object should have.
+	 * Objects must always have a non-empty selection: use FConcertChangeStream_Request::ObjectsToRemove to remove objects.
+	 *
+	 * Any request that would leave an object's property selection empty will result in failure.
+	 */
+	UPROPERTY()
+	FConcertPropertySelection Properties;
+
+	/**
+	 * If Object is pre-existing, this property is optional. If different than None, this will change the class
+	 * If a new object is added, this property is mandatory. Not doing so will fail the request.
+	 *
+	 * Not specifying ClassPath nor Properties is an error: use FConcertChangeStream_Request::ObjectsToRemove to remove objects.
+	 */
+	UPROPERTY()
+	FSoftClassPath ClassPath;
+
+	// Intention: Ideally code dealing with PutObject requests uses these constructors / factory functions.
+	// If a property is added to FReplicatedObjectInfo, only the below code needs to be updated.
+
+	/** @return The PutObject request if New contained sufficient info - empty otherwise */
+	CONCERTSYNCCORE_API static TOptional<FConcertReplication_ChangeStream_PutObject> MakeFromInfo(const FReplicatedObjectInfo& New);
+	/** @return The PutObject request if changing Base to Desired contained sufficient info - empty otherwise */
+	CONCERTSYNCCORE_API static TOptional<FConcertReplication_ChangeStream_PutObject> MakeFromChange(const FReplicatedObjectInfo& Base, const FReplicatedObjectInfo& Desired);
+
+	/** Creates a new object info if there is sufficient data (all fields must be set for this). */
+	CONCERTSYNCCORE_API TOptional<FReplicatedObjectInfo> MakeObjectInfoIfValid() const;
+};
+
+/**
+ * Let's a client change its streams owned on the server.
+ * 
+ * This request is processed atomically: it either succeeds completely or fails completely.
+ * If any sub-change causes a failure, the failure reason will be returned in the response and the client must make a new request.
+ */
+USTRUCT()
+struct FConcertReplication_ChangeStream_Request
+{
+	GENERATED_BODY()
+
+	/**
+	 * Removes objects from pre-existing streams.
+	 * Supplying an object that is not in the specified stream does not cause failure (but is nonsensical).
+	 * 
+	 * If a stream has no registered objects after this operation, it is automatically removed.
+	 * If the requesting client has authority over these objects, authority is removed.
+	 */
+	UPROPERTY()
+	TSet<FObjectInStreamID> ObjectsToRemove;
+
+	/**
+	 * Adds new or modifies preexisting object definitions in pre-existing streams.
+	 *
+	 * If the requesting client and a different client have authority over the same object, you will get a conflict
+	 * if the different client already has authority over one of the properties you're adding here.
+	 * @see FConcertChangeStream_Response::AuthorityConflicts for some examples of conflicts.
+	 *
+	 * If the key identifies a stream that does not exist, the request will fail.
+	 * 
+	 * If the key identifies a stream that is in StreamsToAdd, the request will fail to avoid allowing the construction of
+	 * ambiguous requests, e.g. both StreamsToAdd and ObjectsToPut containing object Foo but with different property selections.
+	 */
+	UPROPERTY()
+	TMap<FObjectInStreamID, FConcertReplication_ChangeStream_PutObject> ObjectsToPut;
+	
+	/**
+	 * New streams to add to the server.
+	 * Fails if any ID overlaps with a pre-existing one.
+	 *
+	 * It is valid to have the same stream in StreamsToRemove: StreamsToAdd is applied after StreamsToRemove resulting
+	 * in the stream's content being replaced.
+	 */
+	UPROPERTY()
+	TArray<FReplicationStreamDescription_NetPacked> StreamsToAdd;
+
+	/**
+	 * Streams to remove from the server.
+	 * Supplying a stream that does not exist does not cause failure (but is nonsensical).
+	 * 
+	 * If the requesting client has authority over any of the objects contained in the stream, authority is removed.
+	 * 
+	 * It is valid to have the same stream in StreamsToAdd: StreamsToAdd is applied after StreamsToRemove resulting
+	 * in the stream's content being replaced.
+	 */
+	UPROPERTY()
+	TSet<FGuid> StreamsToRemove;
+};
+
+UENUM()
+enum class EConcertPutObjectErrorCode : uint8
+{
+	/** Stream that ObjectsToPut referenced was not registered on the server. */
+	UnresolvedStream,
+	/**
+	 * Either PutObject contained no data to update with (ensure either ClassPath or Properties is set),
+	 * or it tried to create a new object with insufficient data (make sure ClassPath and Properties are both specified).
+	 */
+	MissingData
+};
+
+/**
+ * Contains information about why a request failed. This info could be parsed and displayed to the end user as error.
+ * If there is even just one error, the entire request fails and no changes are made server-side.
+ */
+USTRUCT()
+struct FConcertReplication_ChangeStream_Response
+{
+	GENERATED_BODY()
+
+	/**
+	 * Reports dynamic authority errors with ObjectsToPut:
+	 * Changing an object over which a client already has authority can yield unresolvable conflicts for which the entire FConcertChangeStream_Request is rejected.
+	 *
+	 * Let client R be the requester and client A be another client.
+	 * Example 1: Overlapping authority on different clients > Conflict
+	 * - R has authority over Foo's relative rotation property in stream S.
+	 * - A has authority over actor Foo's relative location property in some stream.
+	 * - R requests S to include Foo's relative location property: this conflicts with client A's authority.
+	 *
+	 * Example 2: Overlapping, properties on same client > No conflict
+	 * There is no conflict if another stream owned by the requester already has authority over a property you're adding to a different stream.
+	 * Streams of the same client may overlap properties. While this may not make much sense and actually be performance degrading, it would result in no logical errors so it is allowed.
+	 * Example: R has two streams S1 and S2. S1 contains the transform properties and S2 does not. Both S1 and S2 have authority over object Foo. It is legal to request S2 to contain the transform properties.
+	 */
+	UPROPERTY()
+	TMap<FObjectInStreamID, FReplicatedObjectId> AuthorityConflicts;
+
+	/** Reports semantic errors with ObjectsToPut. */
+	UPROPERTY()
+	TMap<FObjectInStreamID, EConcertPutObjectErrorCode> ObjectsToPutSemanticErrors;
+	
+	/** Streams that were in StreamsToAdd but that were not created. */
+	UPROPERTY()
+	TSet<FGuid> FailedStreamCreation;
+
+	bool IsSuccess() const { return AuthorityConflicts.IsEmpty() && ObjectsToPutSemanticErrors.IsEmpty() && FailedStreamCreation.IsEmpty(); }
+	bool IsFailure() const { return !IsSuccess(); }
+
+	bool WasObjectPutSuccessful(const FObjectInStreamID& Object) const { return !AuthorityConflicts.Contains(Object) && !ObjectsToPutSemanticErrors.Contains(Object); }
+
+	/** If IsFailure(), logs the errors. */
+	CONCERTSYNCCORE_API void LogErrors(FOutputDevice& OutputDevice) const;
+};
\ No newline at end of file
diff --git a/Engine/Plugins/Developer/Concert/ConcertSync/ConcertSyncCore/Source/ConcertSyncCore/Public/Replication/Messages/ClientQuery.h b/Engine/Plugins/Developer/Concert/ConcertSync/ConcertSyncCore/Source/ConcertSyncCore/Public/Replication/Messages/ClientQuery.h
new file mode 100644
--- /dev/null
+++ b/Engine/Plugins/Developer/Concert/ConcertSync/ConcertSyncCore/Source/ConcertSyncCore/Public/Replication/Messages/ClientQuery.h
@@ -0,0 +1,70 @@
+// Copyright Epic Games, Inc. All Rights Reserved.
+
+#pragma once
+
+#include "Replication/Data/ClientQueriedInfo.h"
+#include "ClientQuery.generated.h"
+
+/** Flags to affect FConcertQueryClientStreams_Request */
+UENUM()
+enum class EConcertQueryClientStreamFlags : uint8
+{
+	None = 0,
+
+	/**
+	 * If set, FReplicationClientQueriedInfo::Streams will be empty.
+	 */
+	SkipStreamInfo = 1 << 0,
+
+	/**
+	 * Optimization. The returned streams will not contain the properties that are sent.
+	 * Set this if the request only cares for the objects being replicated but not its content.
+	 * Only has an effect if SkipStreamInfo is not set.
+	 */
+	SkipProperties = 1 << 1,
+	
+	/**
+	 * If set, FReplicationClientQueriedInfo::Authority will be empty.
+	 */
+	SkipAuthority = 1 << 2
+};
+ENUM_CLASS_FLAGS(EConcertQueryClientStreamFlags)
+
+/**
+ * Allows a client query about another client's registered streams.
+ *
+ * Currently, clients must continuously poll to detect when a client changes this data.
+ * TODO: Add an event with which clients can be notified when another client changes their profile.
+ * TODO: Alternative to an event, we can add a version int which is incremented every time the info changes.
+ */
+USTRUCT()
+struct FConcertReplication_QueryReplicationInfo_Request
+{
+	GENERATED_BODY()
+
+	/** The clients to query about */
+	UPROPERTY()
+	TSet<FGuid> ClientEndpointIds;
+
+	UPROPERTY()
+	EConcertQueryClientStreamFlags QueryFlags = EConcertQueryClientStreamFlags::None;
+};
+
+USTRUCT()
+struct FConcertReplication_QueryReplicationInfo_Response
+{
+	GENERATED_BODY()
+
+	/**
+	 * Binds client participating in the replication session to their info.
+	 * 
+	 * @note This contains only the client endpoint IDs that are in the replication session at the time of processing the request.
+	 * That means not every endpoint ID that was in the request's ClientEndpointIds will necessarily be in this result. For example,
+	 * if client A sends a FConcertQueryReplicationInfo_Request while client B disconnects, client B may not show up in ClientInfo.
+	 * 
+	 * Key: Client endpoint ID.
+	 * Value: Info about the client.
+	 */
+	UPROPERTY()
+	TMap<FGuid, FReplicationClientQueriedInfo> ClientInfo;
+};
\ No newline at end of file
diff --git a/Engine/Plugins/Developer/Concert/ConcertSync/ConcertSyncCore/Source/ConcertSyncCore/Public/Replication/Messages/ConcertReplicationEvents.h b/Engine/Plugins/Developer/Concert/ConcertSync/ConcertSyncCore/Source/ConcertSyncCore/Public/Replication/Messages/ConcertReplicationEvents.h
deleted file mode 100644
diff --git a/Engine/Plugins/Developer/Concert/ConcertSync/ConcertSyncCore/Source/ConcertSyncCore/Public/Replication/Messages/ConcertReplicationHandshakeMessages.h b/Engine/Plugins/Developer/Concert/ConcertSync/ConcertSyncCore/Source/ConcertSyncCore/Public/Replication/Messages/ConcertReplicationHandshakeMessages.h
deleted file mode 100644
diff --git a/Engine/Plugins/Developer/Concert/ConcertSync/ConcertSyncCore/Source/ConcertSyncCore/Public/Replication/Messages/Handshake.h b/Engine/Plugins/Developer/Concert/ConcertSync/ConcertSyncCore/Source/ConcertSyncCore/Public/Replication/Messages/Handshake.h
new file mode 100644
--- /dev/null
+++ b/Engine/Plugins/Developer/Concert/ConcertSync/ConcertSyncCore/Source/ConcertSyncCore/Public/Replication/Messages/Handshake.h
@@ -0,0 +1,96 @@
+// Copyright Epic Games, Inc. All Rights Reserved.
+
+#pragma once
+
+#include "Replication/Data/ReplicationStreamDescription.h"
+#include "Handshake.generated.h"
+
+class FText;
+
+UENUM()
+enum class EJoinReplicationErrorCode : uint8
+{
+	Success = 0,
+	/** Indicates generic error obtaining a response from the server. Responses are default constructed if an error occurs. This is the default value. */
+	NetworkError = 1,
+
+	// Client-side error codes:
+	
+	/** Cancelled while initialization was in progress. Most likely LeaveReplicationSession was called while handshake was in progress. */
+	Cancelled = 2,
+	/** Another request is already in progress. */
+	AlreadyInProgress = 3,
+
+	// Server error codes:
+
+	// Session:
+	/**
+	 * Client is not in any session but has requested to join a replication session.
+	 * This case should never happen under normal conditions e.g. when using Multi User Browser, but malicious users could make invalid requests.
+	 */
+	NotInAnyConcertSession = 4,
+	/** Client has already joined a replication session. Send FConcertLeaveReplicationSessionEvent first. */
+	AlreadyInSession = 5,
+
+	// Invalid input:
+	/** One of the specified classes was null or otherwise invalid */
+	InvalidClass = 6,
+	/** An object selection contained a property twice. */
+	DuplicateProperty = 7,
+	/** Your input contained the same stream id twice. */
+	DuplicateStreamId = 8,
+	/** Failed to unpack the stream, likely because one of the attribute classes could not be resolved. */
+	FailedToUnpackStream = 9,
+	
+	MaxPlusOne,
+	/** The maximum possible entry value */
+	Max = MaxPlusOne - 1
+};
+
+namespace UE::ConcertSyncCore::Replication
+{
+	CONCERTSYNCCORE_API FString LexJoinErrorCode(EJoinReplicationErrorCode ErrorCode);
+}
+
+/**
+ * Sent by client to request joining a replication session.
+ * This is sent by clients want to send or receive replicated data, or both.
+ */
+USTRUCT()
+struct FConcertReplication_Join_Request
+{
+	GENERATED_BODY()
+
+	/** The data the client offers to send. */
+	UPROPERTY()
+	TArray<FReplicationStreamDescription_NetPacked> Streams;
+};
+
+/**
+ * Sent by server to accept the previously sent replication session.
+ */
+USTRUCT()
+struct FConcertReplication_Join_Response
+{
+	GENERATED_BODY()
+
+	/**
+	 * Anything other than Success indicates failure.
+	 * 
+	 * Remember that Concert default constructs a response when it fails to obtain a response.
+	 * Hence NetworkError means timeout, etc.
+	 */
+	UPROPERTY()
+	EJoinReplicationErrorCode ErrorCode = EJoinReplicationErrorCode::NetworkError;
+
+	/** More information about ErrorCode to help the user resolve the issue. */
+	UPROPERTY()
+	FString DetailedErrorMessage;
+};
+
+/** Sent by client to notify server that no more replication data should be sent. */
+USTRUCT()
+struct FConcertReplication_LeaveEvent
+{
+	GENERATED_BODY()
+};
\ No newline at end of file
diff --git a/Engine/Plugins/Developer/Concert/ConcertSync/ConcertSyncCore/Source/ConcertSyncCore/Public/Replication/Messages/ObjectReplication.h b/Engine/Plugins/Developer/Concert/ConcertSync/ConcertSyncCore/Source/ConcertSyncCore/Public/Replication/Messages/ObjectReplication.h
new file mode 100644
--- /dev/null
+++ b/Engine/Plugins/Developer/Concert/ConcertSync/ConcertSyncCore/Source/ConcertSyncCore/Public/Replication/Messages/ObjectReplication.h
@@ -0,0 +1,67 @@
+// Copyright Epic Games, Inc. All Rights Reserved.
+
+#pragma once
+
+#include "ConcertMessageData.h"
+#include "Replication/Data/ClientQueriedInfo.h"
+#include "Replication/Data/ObjectIds.h"
+#include "Misc/Optional.h"
+#include "ObjectReplication.generated.h"
+
+class FOutputDevice;
+
+/** Contains data to be applied to a replicated object. */
+USTRUCT()
+struct FConcertReplication_ObjectReplicationEvent
+{
+	GENERATED_BODY()
+
+	/**
+	 * The object to serialize into.
+	 * TODO: To reduce network load, replace this with an int32 that maps into FObjectToIdTable (custom) / UPackageMap (engine)  shared by all clients and the server.
+	 */
+	UPROPERTY()
+	FSoftObjectPath ReplicatedObject;
+	
+	/** Contains another struct as payload, such as FConcertFullObjectReplicationData. */
+	UPROPERTY()
+	FConcertSessionSerializedPayload SerializedPayload;
+};
+
+/** Contains data produced by a specific replication stream. */
+USTRUCT()
+struct FConcertReplication_StreamReplicationEvent
+{
+	GENERATED_BODY()
+
+	/** The stream that produced this data. */
+	UPROPERTY()
+	FGuid StreamId; // TODO DP UE-193261: This field is only relevant to the server and useless for receiving clients.
+
+	/** The objects replicated by this event */
+	UPROPERTY()
+	TArray<FConcertReplication_ObjectReplicationEvent> ReplicatedObjects;
+
+	FConcertReplication_StreamReplicationEvent() = default;
+	explicit FConcertReplication_StreamReplicationEvent(const FGuid& StreamId)
+		: StreamId(StreamId)
+	{}
+};
+
+/**
+ * Contains multiple objects to be replicated.
+ * Sent from
+ *  - Authoritative client to server
+ *  - Server to clients. Every event may be customized for each client (clients can choose to ignore object updates).
+ *
+ *  TODO DP UE-193261: This should be split up into FConcertBatchReplicationEvent_ToServer and FConcertBatchReplicationEvent_ToClient; clients do not need FConcertStreamReplicationEvent::StreamId. 
+ */
+USTRUCT()
+struct FConcertReplication_BatchReplicationEvent 
+{
+	GENERATED_BODY()
+	
+	/** The objects replicated by this event */
+	UPROPERTY()
+	TArray<FConcertReplication_StreamReplicationEvent> Streams;
+};
\ No newline at end of file
diff --git a/Engine/Plugins/Developer/Concert/ConcertSync/ConcertSyncCore/Source/ConcertSyncCore/Public/Replication/Processing/ObjectReplicationCache.h b/Engine/Plugins/Developer/Concert/ConcertSync/ConcertSyncCore/Source/ConcertSyncCore/Public/Replication/Processing/ObjectReplicationCache.h
--- a/Engine/Plugins/Developer/Concert/ConcertSync/ConcertSyncCore/Source/ConcertSyncCore/Public/Replication/Processing/ObjectReplicationCache.h
+++ b/Engine/Plugins/Developer/Concert/ConcertSync/ConcertSyncCore/Source/ConcertSyncCore/Public/Replication/Processing/ObjectReplicationCache.h
@@ -8,7 +8,7 @@
 #include "Templates/Function.h"
 #include "Templates/SharedPointer.h"
 
-struct FConcertObjectReplicationEvent;
+struct FConcertReplication_ObjectReplicationEvent;
 struct FSoftObjectPath;
 struct FGuid;
 
@@ -33,7 +33,7 @@ namespace UE::ConcertSyncCore
 		 * The user can keep hold of Data until it is used, at which point it just let's Data get out of scope.
 		 * If new data is received while this user is referencing Data, Data will be combined to contain any new data.
 		 */
-		virtual void OnDataCached(const FReplicatedObjectId& Object, TSharedRef<const FConcertObjectReplicationEvent> Data) = 0;
+		virtual void OnDataCached(const FReplicatedObjectId& Object, TSharedRef<const FConcertReplication_ObjectReplicationEvent> Data) = 0;
 
 		virtual ~IReplicationCacheUser() = default;
 	};
@@ -63,7 +63,7 @@ namespace UE::ConcertSyncCore
 		 * @param ObjectReplicationEvent The data that was replicated
 		 * @return The number of cache users that accepted this event. 
 		 */
-		int32 StoreUntilConsumed(const FGuid& SendingEndpointId, const FGuid& OriginStreamId, const FConcertObjectReplicationEvent& ObjectReplicationEvent);
+		int32 StoreUntilConsumed(const FGuid& SendingEndpointId, const FGuid& OriginStreamId, const FConcertReplication_ObjectReplicationEvent& ObjectReplicationEvent);
 
 		/** Registers a new user, which will start receiving data for any new data received from now on. */
 		void RegisterDataCacheUser(TSharedRef<IReplicationCacheUser> User);
@@ -89,7 +89,7 @@ namespace UE::ConcertSyncCore
 			 * The intention is that as soon as IReplicationCacheUser has finished using a data event, it will receive a
 			 * new instance: we do not want IReplicationCacheUsers to have large histories of events being combined into them.
 			 */
-			TMap<TWeakPtr<IReplicationCacheUser>, TWeakPtr<FConcertObjectReplicationEvent>> DataInUse;
+			TMap<TWeakPtr<IReplicationCacheUser>, TWeakPtr<FConcertReplication_ObjectReplicationEvent>> DataInUse;
 		};
 		/** Maps every object to the events cached for it. */
 		TMap<FObjectInStreamID, FObjectCache> Cache; 
diff --git a/Engine/Plugins/Developer/Concert/ConcertSync/ConcertSyncCore/Source/ConcertSyncCore/Public/Replication/Processing/ObjectReplicationReceiver.h b/Engine/Plugins/Developer/Concert/ConcertSync/ConcertSyncCore/Source/ConcertSyncCore/Public/Replication/Processing/ObjectReplicationReceiver.h
--- a/Engine/Plugins/Developer/Concert/ConcertSync/ConcertSyncCore/Source/ConcertSyncCore/Public/Replication/Processing/ObjectReplicationReceiver.h
+++ b/Engine/Plugins/Developer/Concert/ConcertSync/ConcertSyncCore/Source/ConcertSyncCore/Public/Replication/Processing/ObjectReplicationReceiver.h
@@ -5,10 +5,10 @@
 #include "Templates/SharedPointer.h"
 
 class IConcertSession;
-struct FConcertBatchReplicationEvent;
-struct FConcertObjectReplicationEvent;
+struct FConcertReplication_BatchReplicationEvent;
+struct FConcertReplication_ObjectReplicationEvent;
 struct FConcertSessionContext;
-struct FConcertStreamReplicationEvent;
+struct FConcertReplication_StreamReplicationEvent;
 
 namespace UE::ConcertSyncCore
 {
@@ -25,7 +25,7 @@ namespace UE::ConcertSyncCore
 	protected:
 
 		/** Whether the object should be processed. */
-		virtual bool ShouldAcceptObject(const FConcertSessionContext& SessionContext, const FConcertStreamReplicationEvent& StreamEvent, const FConcertObjectReplicationEvent& ObjectEvent) const { return true; }
+		virtual bool ShouldAcceptObject(const FConcertSessionContext& SessionContext, const FConcertReplication_StreamReplicationEvent& StreamEvent, const FConcertReplication_ObjectReplicationEvent& ObjectEvent) const { return true; }
 
 	private:
 
@@ -34,6 +34,6 @@ namespace UE::ConcertSyncCore
 		/** Where received data is stored. */
 		TSharedRef<FObjectReplicationCache> ReplicationCache;
 
-		void HandleBatchReplicationEvent(const FConcertSessionContext& SessionContext, const FConcertBatchReplicationEvent& Event);
+		void HandleBatchReplicationEvent(const FConcertSessionContext& SessionContext, const FConcertReplication_BatchReplicationEvent& Event);
 	};
 }
diff --git a/Engine/Plugins/Developer/Concert/ConcertSync/ConcertSyncCore/Source/ConcertSyncCore/Public/Replication/Processing/ObjectReplicationSender.h b/Engine/Plugins/Developer/Concert/ConcertSync/ConcertSyncCore/Source/ConcertSyncCore/Public/Replication/Processing/ObjectReplicationSender.h
--- a/Engine/Plugins/Developer/Concert/ConcertSync/ConcertSyncCore/Source/ConcertSyncCore/Public/Replication/Processing/ObjectReplicationSender.h
+++ b/Engine/Plugins/Developer/Concert/ConcertSync/ConcertSyncCore/Source/ConcertSyncCore/Public/Replication/Processing/ObjectReplicationSender.h
@@ -3,7 +3,7 @@
 #pragma once
 
 #include "ObjectReplicationProcessor.h"
-#include "Replication/Messages/ConcertReplicationEvents.h"
+#include "Replication/Messages/ObjectReplication.h"
 
 class IConcertSession;
 
@@ -37,6 +37,6 @@ namespace UE::ConcertSyncCore
 		TSharedRef<IConcertSession> Session;
 
 		/** This event is filled in ProcessObjects and finally sent to TargetEndpointId. */
-		FConcertBatchReplicationEvent EventToSend;
+		FConcertReplication_BatchReplicationEvent EventToSend;
 	};
 }
diff --git a/Engine/Plugins/Developer/Concert/ConcertSync/ConcertSyncCore/Source/ConcertSyncCore/Public/Replication/Processing/ReplicationDataQueuer.h b/Engine/Plugins/Developer/Concert/ConcertSync/ConcertSyncCore/Source/ConcertSyncCore/Public/Replication/Processing/ReplicationDataQueuer.h
--- a/Engine/Plugins/Developer/Concert/ConcertSync/ConcertSyncCore/Source/ConcertSyncCore/Public/Replication/Processing/ReplicationDataQueuer.h
+++ b/Engine/Plugins/Developer/Concert/ConcertSync/ConcertSyncCore/Source/ConcertSyncCore/Public/Replication/Processing/ReplicationDataQueuer.h
@@ -33,7 +33,7 @@ namespace UE::ConcertSyncCore
 		//~ End IReplicationDataSource Interface
 		
 		//~ Begin IReplicationCacheUser Interface
-		virtual void OnDataCached(const FReplicatedObjectId& Object, TSharedRef<const FConcertObjectReplicationEvent> Data) override;
+		virtual void OnDataCached(const FReplicatedObjectId& Object, TSharedRef<const FConcertReplication_ObjectReplicationEvent> Data) override;
 		//~ End IReplicationCacheUser Interface
 
 	protected:
@@ -44,7 +44,7 @@ namespace UE::ConcertSyncCore
 	private:
 
 		/** Stores events as they are received. */
-		TMap<FObjectInStreamID, TSharedPtr<const FConcertObjectReplicationEvent>> PendingEvents;
+		TMap<FObjectInStreamID, TSharedPtr<const FConcertReplication_ObjectReplicationEvent>> PendingEvents;
 
 		/** Provides us with replication events and shares them effectively. */
 		TSharedPtr<FObjectReplicationCache> ReplicationCache;
diff --git a/Engine/Plugins/Developer/Concert/ConcertSync/ConcertSyncServer/Source/ConcertSyncServer/Private/Replication/AuthorityManager.cpp b/Engine/Plugins/Developer/Concert/ConcertSync/ConcertSyncServer/Source/ConcertSyncServer/Private/Replication/AuthorityManager.cpp
--- a/Engine/Plugins/Developer/Concert/ConcertSync/ConcertSyncServer/Source/ConcertSyncServer/Private/Replication/AuthorityManager.cpp
+++ b/Engine/Plugins/Developer/Concert/ConcertSync/ConcertSyncServer/Source/ConcertSyncServer/Private/Replication/AuthorityManager.cpp
@@ -6,7 +6,7 @@
 #include "IConcertSession.h"
 #include "Replication/Data/ObjectIds.h"
 #include "Replication/Data/ReplicationStreamDescription.h"
-#include "Replication/Messages/ConcertReplicationEvents.h"
+#include "Replication/Messages/ChangeAuthority.h"
 
 namespace UE::ConcertSyncServer::Replication
 {
@@ -31,12 +31,12 @@ namespace UE::ConcertSyncServer::Replication
 		: Getters(Getters)
 		, Session(MoveTemp(InSession))
 	{
-		Session->RegisterCustomRequestHandler<FConcertChangeAuthority_Request, FConcertChangeAuthority_Response>(this, &FAuthorityManager::HandleChangeAuthorityRequest);
+		Session->RegisterCustomRequestHandler<FConcertReplication_ChangeAuthority_Request, FConcertReplication_ChangeAuthority_Response>(this, &FAuthorityManager::HandleChangeAuthorityRequest);
 	}
 
 	FAuthorityManager::~FAuthorityManager()
 	{
-		Session->UnregisterCustomRequestHandler<FConcertChangeAuthority_Request>();
+		Session->UnregisterCustomRequestHandler<FConcertReplication_ChangeAuthority_Request>();
 	}
 
 	bool FAuthorityManager::HasAuthorityToChange(const FReplicatedObjectId& ObjectChange) const
@@ -144,8 +144,8 @@ namespace UE::ConcertSyncServer::Replication
 
 	EConcertSessionResponseCode FAuthorityManager::HandleChangeAuthorityRequest(
 		const FConcertSessionContext& ConcertSessionContext,
-		const FConcertChangeAuthority_Request& Request,
-		FConcertChangeAuthority_Response& Response
+		const FConcertReplication_ChangeAuthority_Request& Request,
+		FConcertReplication_ChangeAuthority_Response& Response
 		)
 	{
 		// This log does two things: 1. Identify issues in unit tests / at runtime 2. Warn about possibly malicious attempts when the server runs.
diff --git a/Engine/Plugins/Developer/Concert/ConcertSync/ConcertSyncServer/Source/ConcertSyncServer/Private/Replication/AuthorityManager.h b/Engine/Plugins/Developer/Concert/ConcertSync/ConcertSyncServer/Source/ConcertSyncServer/Private/Replication/AuthorityManager.h
--- a/Engine/Plugins/Developer/Concert/ConcertSync/ConcertSyncServer/Source/ConcertSyncServer/Private/Replication/AuthorityManager.h
+++ b/Engine/Plugins/Developer/Concert/ConcertSync/ConcertSyncServer/Source/ConcertSyncServer/Private/Replication/AuthorityManager.h
@@ -9,8 +9,8 @@
 
 class IConcertSession;
 
-struct FConcertChangeAuthority_Response;
-struct FConcertChangeAuthority_Request;
+struct FConcertReplication_ChangeAuthority_Response;
+struct FConcertReplication_ChangeAuthority_Request;
 struct FConcertPropertySelection;
 struct FConcertSessionContext;
 struct FObjectInStreamID;
@@ -100,8 +100,8 @@ namespace UE::ConcertSyncServer::Replication
 
 		EConcertSessionResponseCode HandleChangeAuthorityRequest(
 			const FConcertSessionContext& ConcertSessionContext,
-			const FConcertChangeAuthority_Request& Request,
-			FConcertChangeAuthority_Response& Response
+			const FConcertReplication_ChangeAuthority_Request& Request,
+			FConcertReplication_ChangeAuthority_Response& Response
 			);
 
 		/** Finds a stream registered with the client by its ID. */
diff --git a/Engine/Plugins/Developer/Concert/ConcertSync/ConcertSyncServer/Source/ConcertSyncServer/Private/Replication/ConcertReplicationClient.cpp b/Engine/Plugins/Developer/Concert/ConcertSync/ConcertSyncServer/Source/ConcertSyncServer/Private/Replication/ConcertReplicationClient.cpp
--- a/Engine/Plugins/Developer/Concert/ConcertSync/ConcertSyncServer/Source/ConcertSyncServer/Private/Replication/ConcertReplicationClient.cpp
+++ b/Engine/Plugins/Developer/Concert/ConcertSync/ConcertSyncServer/Source/ConcertSyncServer/Private/Replication/ConcertReplicationClient.cpp
@@ -4,7 +4,7 @@
 
 #include "Processing/ServerReplicationDataQueuer.h"
 #include "Replication/ChangeStreamSharedUtils.h"
-#include "Replication/Messages/ConcertReplicationHandshakeMessages.h"
+#include "Replication/Messages/Handshake.h"
 
 namespace UE::ConcertSyncServer::Replication
 {
@@ -25,7 +25,7 @@ namespace UE::ConcertSyncServer::Replication
 		DataRelay.ProcessObjects(TimeBudget);
 	}
 
-	void FConcertReplicationClient::ApplyValidatedRequest(const FConcertChangeStream_Request& Request)
+	void FConcertReplicationClient::ApplyValidatedRequest(const FConcertReplication_ChangeStream_Request& Request)
 	{
 		// Right now there is nothing further to do but if in future you need to update some client systems of the change, this is the place to do it.
 		ConcertSyncCore::Replication::ChangeStreamUtils::ApplyValidatedRequest(Request, StreamDescriptions);
diff --git a/Engine/Plugins/Developer/Concert/ConcertSync/ConcertSyncServer/Source/ConcertSyncServer/Private/Replication/ConcertReplicationClient.h b/Engine/Plugins/Developer/Concert/ConcertSync/ConcertSyncServer/Source/ConcertSyncServer/Private/Replication/ConcertReplicationClient.h
--- a/Engine/Plugins/Developer/Concert/ConcertSync/ConcertSyncServer/Source/ConcertSyncServer/Private/Replication/ConcertReplicationClient.h
+++ b/Engine/Plugins/Developer/Concert/ConcertSync/ConcertSyncServer/Source/ConcertSyncServer/Private/Replication/ConcertReplicationClient.h
@@ -8,7 +8,7 @@
 
 struct FReplicationStreamDescription;
 struct FConcertReplication_Join_Request;
-struct FConcertChangeStream_Request;
+struct FConcertReplication_ChangeStream_Request;
 
 namespace UE::ConcertSyncCore
 {
@@ -38,7 +38,7 @@ namespace UE::ConcertSyncServer::Replication
 		void ProcessClient(float TimeBudget);
 		
 		/** Updates the StreamDescriptions array with the changes from Request. The request already passed validation and is valid to apply. */
-		void ApplyValidatedRequest(const FConcertChangeStream_Request& Request);
+		void ApplyValidatedRequest(const FConcertReplication_ChangeStream_Request& Request);
 
 		const FGuid& GetClientEndpointId() const { return ClientEndpointId; }
 		const TArray<FReplicationStreamDescription>& GetStreamDescriptions() const { return StreamDescriptions; }
diff --git a/Engine/Plugins/Developer/Concert/ConcertSync/ConcertSyncServer/Source/ConcertSyncServer/Private/Replication/ConcertServerReplicationManager.cpp b/Engine/Plugins/Developer/Concert/ConcertSync/ConcertSyncServer/Source/ConcertSyncServer/Private/Replication/ConcertServerReplicationManager.cpp
--- a/Engine/Plugins/Developer/Concert/ConcertSync/ConcertSyncServer/Source/ConcertSyncServer/Private/Replication/ConcertServerReplicationManager.cpp
+++ b/Engine/Plugins/Developer/Concert/ConcertSync/ConcertSyncServer/Source/ConcertSyncServer/Private/Replication/ConcertServerReplicationManager.cpp
@@ -7,7 +7,9 @@
 #include "IConcertSession.h"
 #include "Replication/ConcertReplicationClient.h"
 #include "Replication/Formats/FullObjectFormat.h"
-#include "Replication/Messages/ConcertReplicationHandshakeMessages.h"
+#include "Replication/Messages/ChangeStream.h"
+#include "Replication/Messages/ClientQuery.h"
+#include "Replication/Messages/Handshake.h"
 #include "Replication/Processing/ObjectReplicationCache.h"
 #include "Util/JoinRequestValidation.h"
 
@@ -21,8 +23,8 @@ namespace UE::ConcertSyncServer::Replication
 		, ReplicationDataReceiver(AuthorityManager, Session, ReplicationCache)
 	{
 		Session->RegisterCustomRequestHandler<FConcertReplication_Join_Request, FConcertReplication_Join_Response>(this, &FConcertServerReplicationManager::HandleJoinReplicationSessionRequest);
-		Session->RegisterCustomRequestHandler<FConcertQueryReplicationInfo_Request, FConcertQueryReplicationInfo_Response>(this, &FConcertServerReplicationManager::HandleQueryReplicationInfoRequest);
-		Session->RegisterCustomRequestHandler<FConcertChangeStream_Request, FConcertChangeStream_Response>(this, &FConcertServerReplicationManager::HandleChangeStreamRequest);
+		Session->RegisterCustomRequestHandler<FConcertReplication_QueryReplicationInfo_Request, FConcertReplication_QueryReplicationInfo_Response>(this, &FConcertServerReplicationManager::HandleQueryReplicationInfoRequest);
+		Session->RegisterCustomRequestHandler<FConcertReplication_ChangeStream_Request, FConcertReplication_ChangeStream_Response>(this, &FConcertServerReplicationManager::HandleChangeStreamRequest);
 		Session->RegisterCustomEventHandler<FConcertReplication_LeaveEvent>(this, &FConcertServerReplicationManager::HandleLeaveReplicationSessionRequest);
 		Session->OnSessionClientChanged().AddRaw(this, &FConcertServerReplicationManager::OnConnectionChanged);
 
@@ -32,7 +34,7 @@ namespace UE::ConcertSyncServer::Replication
 	FConcertServerReplicationManager::~FConcertServerReplicationManager()
 	{
 		Session->UnregisterCustomRequestHandler<FConcertReplication_Join_Response>();
-		Session->UnregisterCustomRequestHandler<FConcertQueryReplicationInfo_Response>();
+		Session->UnregisterCustomRequestHandler<FConcertReplication_QueryReplicationInfo_Response>();
 		Session->UnregisterCustomEventHandler<FConcertReplication_LeaveEvent>(this);
 
 		Session->OnTick().RemoveAll(this);
@@ -117,8 +119,8 @@ namespace UE::ConcertSyncServer::Replication
 
 	EConcertSessionResponseCode FConcertServerReplicationManager::HandleQueryReplicationInfoRequest(
 		const FConcertSessionContext& ConcertSessionContext,
-		const FConcertQueryReplicationInfo_Request& Request,
-		FConcertQueryReplicationInfo_Response& Response
+		const FConcertReplication_QueryReplicationInfo_Request& Request,
+		FConcertReplication_QueryReplicationInfo_Response& Response
 		)
 	{
 		for (const FGuid& EndpointId : Request.ClientEndpointIds)
diff --git a/Engine/Plugins/Developer/Concert/ConcertSync/ConcertSyncServer/Source/ConcertSyncServer/Private/Replication/ConcertServerReplicationManager.h b/Engine/Plugins/Developer/Concert/ConcertSync/ConcertSyncServer/Source/ConcertSyncServer/Private/Replication/ConcertServerReplicationManager.h
--- a/Engine/Plugins/Developer/Concert/ConcertSync/ConcertSyncServer/Source/ConcertSyncServer/Private/Replication/ConcertServerReplicationManager.h
+++ b/Engine/Plugins/Developer/Concert/ConcertSync/ConcertSyncServer/Source/ConcertSyncServer/Private/Replication/ConcertServerReplicationManager.h
@@ -7,21 +7,25 @@
 #include "ConcertReplicationClient.h"
 #include "Processing/ServerObjectReplicationReceiver.h"
 #include "Replication/IConcertServerReplicationManager.h"
-#include "Replication/Messages/ConcertReplicationHandshakeMessages.h"
+#include "Replication/Messages/Handshake.h"
 
 #include "Templates/SharedPointer.h"
 #include "Templates/Tuple.h"
 #include "Templates/UnrealTemplate.h"
 
+class IConcertClientReplicationBridge;
+class IConcertServerSession;
+
+struct FConcertReplication_ChangeStream_Response;
+struct FConcertReplication_QueryReplicationInfo_Response;
+struct FConcertReplication_QueryReplicationInfo_Request;
+
 namespace UE::ConcertSyncCore
 {
 	class IObjectReplicationFormat;
 	class FObjectReplicationCache;
 }
 
-class IConcertClientReplicationBridge;
-class IConcertServerSession;
-
 namespace UE::ConcertSyncServer::Replication
 {
 	class FAuthorityManager;
@@ -73,14 +77,14 @@ namespace UE::ConcertSyncServer::Replication
 		EConcertSessionResponseCode InternalHandleJoinReplicationSessionRequest(const FConcertSessionContext& ConcertSessionContext, const FConcertReplication_Join_Request& Request, FConcertReplication_Join_Response& Response);
 
 		// Querying
-		EConcertSessionResponseCode HandleQueryReplicationInfoRequest(const FConcertSessionContext& ConcertSessionContext, const FConcertQueryReplicationInfo_Request& Request, FConcertQueryReplicationInfo_Response& Response);
+		EConcertSessionResponseCode HandleQueryReplicationInfoRequest(const FConcertSessionContext& ConcertSessionContext, const FConcertReplication_QueryReplicationInfo_Request& Request, FConcertReplication_QueryReplicationInfo_Response& Response);
 		/** Gets all registered streams and optionally removes the properties. */
 		TArray<FSharedReplicationStreamDescription> BuildClientStreamInfo(const FConcertReplicationClient& Client, bool bSkipProperties) const;
 		/** Maps the client's streams to the objects in that stream the client has taken authority over. */
 		TArray<FReplicationAuthorityInfo> BuildClientAuthorityInfo(const FConcertReplicationClient& Client) const;
 
 		// Changing streams
-		EConcertSessionResponseCode HandleChangeStreamRequest(const FConcertSessionContext& ConcertSessionContext, const FConcertChangeStream_Request& Request, FConcertChangeStream_Response& Response);
+		EConcertSessionResponseCode HandleChangeStreamRequest(const FConcertSessionContext& ConcertSessionContext, const FConcertReplication_ChangeStream_Request& Request, FConcertReplication_ChangeStream_Response& Response);
 
 		// Leaving
 		void HandleLeaveReplicationSessionRequest(const FConcertSessionContext& ConcertSessionContext, const FConcertReplication_LeaveEvent& EventData);
diff --git a/Engine/Plugins/Developer/Concert/ConcertSync/ConcertSyncServer/Source/ConcertSyncServer/Private/Replication/ConcertServerReplicationManager_ChangeStream.cpp b/Engine/Plugins/Developer/Concert/ConcertSync/ConcertSyncServer/Source/ConcertSyncServer/Private/Replication/ConcertServerReplicationManager_ChangeStream.cpp
--- a/Engine/Plugins/Developer/Concert/ConcertSync/ConcertSyncServer/Source/ConcertSyncServer/Private/Replication/ConcertServerReplicationManager_ChangeStream.cpp
+++ b/Engine/Plugins/Developer/Concert/ConcertSync/ConcertSyncServer/Source/ConcertSyncServer/Private/Replication/ConcertServerReplicationManager_ChangeStream.cpp
@@ -7,8 +7,8 @@
 #include "IConcertSession.h"
 #include "Replication/ChangeStreamSharedUtils.h"
 #include "Replication/ConcertReplicationClient.h"
-#include "Replication/Messages/ConcertReplicationHandshakeMessages.h"
-#include "Replication/Processing/ObjectReplicationCache.h"
+#include "Replication/Messages/ChangeStream.h"
+#include "Replication/Messages/Handshake.h"
 
 namespace UE::ConcertSyncServer::Replication
 {
@@ -23,9 +23,9 @@ namespace UE::ConcertSyncServer::Replication
 		}
 		
 		/** Validates that ObjectsToPut writes only to pre-existing streams. */
-		static void ValidatePutObjectsRequestSemantics(const FConcertChangeStream_Request& Request, const FConcertReplicationClient& Client, FConcertChangeStream_Response& OutResponse)
+		static void ValidatePutObjectsRequestSemantics(const FConcertReplication_ChangeStream_Request& Request, const FConcertReplicationClient& Client, FConcertReplication_ChangeStream_Response& OutResponse)
 		{
-			const auto PutObjectHasEnoughData = [](const FReplicationStreamDescription* ExistingStream, const FSoftObjectPath& ChangedObjectPath, const FConcertChangeStream_PutObject& PutObject)
+			const auto PutObjectHasEnoughData = [](const FReplicationStreamDescription* ExistingStream, const FSoftObjectPath& ChangedObjectPath, const FConcertReplication_ChangeStream_PutObject& PutObject)
 			{
 				const bool bHasProperties = !PutObject.Properties.ReplicatedProperties.IsEmpty();
 				const bool bHasClassPath = !PutObject.ClassPath.IsNull();
@@ -36,7 +36,7 @@ namespace UE::ConcertSyncServer::Replication
 				return bHasEnoughData;
 			};
 			
-			for (const TPair<FObjectInStreamID, FConcertChangeStream_PutObject>& Change : Request.ObjectsToPut)
+			for (const TPair<FObjectInStreamID, FConcertReplication_ChangeStream_PutObject>& Change : Request.ObjectsToPut)
 			{
 				const FObjectInStreamID ChangedObject = Change.Key;
 				
@@ -48,7 +48,7 @@ namespace UE::ConcertSyncServer::Replication
 					// FReplicatedObjectInfo must have non-empty values for its members. Hence we must check that
 					// 1. if creating a new entry, both fields are given
 					// 2. if writing to pre-existing entry, at least one field is given
-					const FConcertChangeStream_PutObject& PutObject = Change.Value; 
+					const FConcertReplication_ChangeStream_PutObject& PutObject = Change.Value; 
 					if (!PutObjectHasEnoughData(ExistingStream, ChangedObject.Object, PutObject))
 					{
 						OutResponse.ObjectsToPutSemanticErrors.Add(ChangedObject, EConcertPutObjectErrorCode::MissingData);
@@ -63,7 +63,7 @@ namespace UE::ConcertSyncServer::Replication
 		}
 
 		/** Checks that StreamsToAdd do not conflict with pre-existing ones and that all IDs in the request are also unique. */
-		static void ValidateAddedStreamsAreUnique(const FConcertChangeStream_Request& Request, const FConcertReplicationClient& Client, FConcertChangeStream_Response& OutResponse)
+		static void ValidateAddedStreamsAreUnique(const FConcertReplication_ChangeStream_Request& Request, const FConcertReplicationClient& Client, FConcertReplication_ChangeStream_Response& OutResponse)
 		{
 			TSet<FGuid> DuplicateEntryDetection;
 			for (const FReplicationStreamDescription_NetPacked& NewStream : Request.StreamsToAdd)
@@ -84,10 +84,10 @@ namespace UE::ConcertSyncServer::Replication
 		}
 
 		/** If the requesting client has authority over a changed object, checks that no other client has authority over the properties being added. */
-		static void LookForAuthorityConflicts(const FConcertChangeStream_Request& Request, const FConcertReplicationClient& Client, const FAuthorityManager& AuthorityManager, FConcertChangeStream_Response& OutResponse)
+		static void LookForAuthorityConflicts(const FConcertReplication_ChangeStream_Request& Request, const FConcertReplicationClient& Client, const FAuthorityManager& AuthorityManager, FConcertReplication_ChangeStream_Response& OutResponse)
 		{
 			const FGuid ClientEndpointId = Client.GetClientEndpointId();
-			for (const TPair<FObjectInStreamID, FConcertChangeStream_PutObject>& PutObjectPair : Request.ObjectsToPut)
+			for (const TPair<FObjectInStreamID, FConcertReplication_ChangeStream_PutObject>& PutObjectPair : Request.ObjectsToPut)
 			{
 				// No conflict possible if requesting client does not have authority over the changed object
 				const FReplicatedObjectId ReplicatedObjectInfo { { PutObjectPair.Key }, ClientEndpointId };
@@ -116,7 +116,7 @@ namespace UE::ConcertSyncServer::Replication
 		}
 
 		/** Stream description can have UObjects attached which fail to serialize on the receiving end. If so, reject. */
-		static void ValidateUnpacking(const FConcertChangeStream_Request& Request, FConcertChangeStream_Response& OutResponse)
+		static void ValidateUnpacking(const FConcertReplication_ChangeStream_Request& Request, FConcertReplication_ChangeStream_Response& OutResponse)
 		{
 			for (const FReplicationStreamDescription_NetPacked& ToUnpack : Request.StreamsToAdd)
 			{
@@ -128,7 +128,7 @@ namespace UE::ConcertSyncServer::Replication
 		}
 
 		/** Checks whether this request is valid to apply. */
-		static bool ShouldAcceptRequest(const FConcertChangeStream_Request& Request, const FConcertReplicationClient& Client, const FAuthorityManager& AuthorityManager, FConcertChangeStream_Response& OutResponse)
+		static bool ShouldAcceptRequest(const FConcertReplication_ChangeStream_Request& Request, const FConcertReplicationClient& Client, const FAuthorityManager& AuthorityManager, FConcertReplication_ChangeStream_Response& OutResponse)
 		{
 			ValidatePutObjectsRequestSemantics(Request, Client, OutResponse);
 			ValidateAddedStreamsAreUnique(Request, Client, OutResponse);
@@ -140,8 +140,8 @@ namespace UE::ConcertSyncServer::Replication
 
 	EConcertSessionResponseCode FConcertServerReplicationManager::HandleChangeStreamRequest(
 		const FConcertSessionContext& ConcertSessionContext,
-		const FConcertChangeStream_Request& Request,
-		FConcertChangeStream_Response& Response)
+		const FConcertReplication_ChangeStream_Request& Request,
+		FConcertReplication_ChangeStream_Response& Response)
 	{
 		const FGuid SendingClientId = ConcertSessionContext.SourceEndpointId;
 		const TSharedRef<FConcertReplicationClient>* SendingClient = Clients.Find(SendingClientId);
diff --git a/Engine/Plugins/Developer/Concert/ConcertSync/ConcertSyncServer/Source/ConcertSyncServer/Private/Replication/Processing/ServerObjectReplicationReceiver.cpp b/Engine/Plugins/Developer/Concert/ConcertSync/ConcertSyncServer/Source/ConcertSyncServer/Private/Replication/Processing/ServerObjectReplicationReceiver.cpp
--- a/Engine/Plugins/Developer/Concert/ConcertSync/ConcertSyncServer/Source/ConcertSyncServer/Private/Replication/Processing/ServerObjectReplicationReceiver.cpp
+++ b/Engine/Plugins/Developer/Concert/ConcertSync/ConcertSyncServer/Source/ConcertSyncServer/Private/Replication/Processing/ServerObjectReplicationReceiver.cpp
@@ -5,7 +5,7 @@
 #include "IConcertSessionHandler.h"
 #include "Replication/AuthorityManager.h"
 #include "Replication/Data/ObjectIds.h"
-#include "Replication/Messages/ConcertReplicationEvents.h"
+#include "Replication/Messages/ObjectReplication.h"
 
 namespace UE::ConcertSyncServer::Replication
 {
@@ -20,8 +20,8 @@ namespace UE::ConcertSyncServer::Replication
 
 	bool FServerObjectReplicationReceiver::ShouldAcceptObject(
 		const FConcertSessionContext& SessionContext,
-		const FConcertStreamReplicationEvent& StreamEvent,
-		const FConcertObjectReplicationEvent& ObjectEvent
+		const FConcertReplication_StreamReplicationEvent& StreamEvent,
+		const FConcertReplication_ObjectReplicationEvent& ObjectEvent
 		) const
 	{
 		const FReplicatedObjectId ReplicatedObjectInfo { { StreamEvent.StreamId, ObjectEvent.ReplicatedObject }, SessionContext.SourceEndpointId };
diff --git a/Engine/Plugins/Developer/Concert/ConcertSync/ConcertSyncServer/Source/ConcertSyncServer/Private/Replication/Processing/ServerObjectReplicationReceiver.h b/Engine/Plugins/Developer/Concert/ConcertSync/ConcertSyncServer/Source/ConcertSyncServer/Private/Replication/Processing/ServerObjectReplicationReceiver.h
--- a/Engine/Plugins/Developer/Concert/ConcertSync/ConcertSyncServer/Source/ConcertSyncServer/Private/Replication/Processing/ServerObjectReplicationReceiver.h
+++ b/Engine/Plugins/Developer/Concert/ConcertSync/ConcertSyncServer/Source/ConcertSyncServer/Private/Replication/Processing/ServerObjectReplicationReceiver.h
@@ -19,7 +19,7 @@ namespace UE::ConcertSyncServer::Replication
 	protected:
 
 		//~ Begin FObjectReplicationReceiver Interface
-		virtual bool ShouldAcceptObject(const FConcertSessionContext& SessionContext, const FConcertStreamReplicationEvent& StreamEvent, const FConcertObjectReplicationEvent& ObjectEvent) const override;
+		virtual bool ShouldAcceptObject(const FConcertSessionContext& SessionContext, const FConcertReplication_StreamReplicationEvent& StreamEvent, const FConcertReplication_ObjectReplicationEvent& ObjectEvent) const override;
 		//~ End FObjectReplicationReceiver Interface
 
 	private:
diff --git a/Engine/Plugins/Developer/Concert/ConcertSync/ConcertSyncServer/Source/ConcertSyncServer/Private/Replication/Util/JoinRequestValidation.cpp b/Engine/Plugins/Developer/Concert/ConcertSync/ConcertSyncServer/Source/ConcertSyncServer/Private/Replication/Util/JoinRequestValidation.cpp
--- a/Engine/Plugins/Developer/Concert/ConcertSync/ConcertSyncServer/Source/ConcertSyncServer/Private/Replication/Util/JoinRequestValidation.cpp
+++ b/Engine/Plugins/Developer/Concert/ConcertSync/ConcertSyncServer/Source/ConcertSyncServer/Private/Replication/Util/JoinRequestValidation.cpp
@@ -5,7 +5,7 @@
 #include "Replication/ConcertReplicationClient.h"
 #include "Replication/Data/ObjectReplicationMap.h"
 #include "Replication/Data/ReplicationStreamDescription.h"
-#include "Replication/Messages/ConcertReplicationHandshakeMessages.h"
+#include "Replication/Messages/Handshake.h"
 
 namespace UE::ConcertSyncServer::Replication
 {
diff --git a/Engine/Plugins/Developer/Concert/ConcertSync/ConcertSyncTest/Source/ConcertSyncTest/Private/Replication/Components/ReplicationCacheTests.cpp b/Engine/Plugins/Developer/Concert/ConcertSync/ConcertSyncTest/Source/ConcertSyncTest/Private/Replication/Components/ReplicationCacheTests.cpp
--- a/Engine/Plugins/Developer/Concert/ConcertSync/ConcertSyncTest/Source/ConcertSyncTest/Private/Replication/Components/ReplicationCacheTests.cpp
+++ b/Engine/Plugins/Developer/Concert/ConcertSync/ConcertSyncTest/Source/ConcertSyncTest/Private/Replication/Components/ReplicationCacheTests.cpp
@@ -1,7 +1,7 @@
 // Copyright Epic Games, Inc. All Rights Reserved.
 
 #include "Replication/Formats/IObjectReplicationFormat.h"
-#include "Replication/Messages/ConcertReplicationEvents.h"
+#include "Replication/Messages/ObjectReplication.h"
 #include "Replication/Processing/ObjectReplicationCache.h"
 #include "Replication/Data/ObjectIds.h"
 #include "Replication/TestReflectionObject.h"
@@ -46,7 +46,7 @@ namespace UE::ConcertSyncTests
 			}
 			virtual void ApplyReplicationEvent(UObject& Object, const FConcertSessionSerializedPayload& Payload) override { return NotMocked<void>(); }
 
-			static FConcertObjectReplicationEvent CreateEvent(FSoftObjectPath Object, float Value)
+			static FConcertReplication_ObjectReplicationEvent CreateEvent(FSoftObjectPath Object, float Value)
 			{
 				FNativeStruct Data{ Value };
 				FConcertSessionSerializedPayload Payload;
@@ -62,7 +62,7 @@ namespace UE::ConcertSyncTests
 			FAutomationTestBase& Test;
 			FObjectInStreamID AllowedObject;
 			EReplicationCacheTestFlags Flags;
-			TSharedPtr<const FConcertObjectReplicationEvent> CachedData;
+			TSharedPtr<const FConcertReplication_ObjectReplicationEvent> CachedData;
 
 			FTestReplicationCacheUser(FAutomationTestBase& Test, FObjectInStreamID AllowedObject, EReplicationCacheTestFlags Flags = EReplicationCacheTestFlags::None)
 				: Test(Test)
@@ -75,7 +75,7 @@ namespace UE::ConcertSyncTests
 				return !EnumHasAnyFlags(Flags, EReplicationCacheTestFlags::NeverReceive) && AllowedObject == Object;
 			}
 			
-			virtual void OnDataCached(const FReplicatedObjectId& Object, TSharedRef<const FConcertObjectReplicationEvent> Data) override
+			virtual void OnDataCached(const FReplicatedObjectId& Object, TSharedRef<const FConcertReplication_ObjectReplicationEvent> Data) override
 			{
 				if (EnumHasAnyFlags(Flags, EReplicationCacheTestFlags::NeverReceive))
 				{
@@ -137,10 +137,10 @@ namespace UE::ConcertSyncTests
 		Cache->RegisterDataCacheUser(User_ConsumeManually);
 		Cache->RegisterDataCacheUser(User_InstantConsume);
 
-		const FConcertObjectReplicationEvent Event_5 = FTestReplicationFormat::CreateEvent(ObjectPath, 5.f);
-		const FConcertObjectReplicationEvent Event_10 = FTestReplicationFormat::CreateEvent(ObjectPath, 10.f);
-		const FConcertObjectReplicationEvent Event_20 = FTestReplicationFormat::CreateEvent(ObjectPath, 20.f);
-		const FConcertObjectReplicationEvent Event_100 = FTestReplicationFormat::CreateEvent(ObjectPath, 100.f);
+		const FConcertReplication_ObjectReplicationEvent Event_5 = FTestReplicationFormat::CreateEvent(ObjectPath, 5.f);
+		const FConcertReplication_ObjectReplicationEvent Event_10 = FTestReplicationFormat::CreateEvent(ObjectPath, 10.f);
+		const FConcertReplication_ObjectReplicationEvent Event_20 = FTestReplicationFormat::CreateEvent(ObjectPath, 20.f);
+		const FConcertReplication_ObjectReplicationEvent Event_100 = FTestReplicationFormat::CreateEvent(ObjectPath, 100.f);
 
 
 		
diff --git a/Engine/Plugins/Developer/Concert/ConcertSync/ConcertSyncTest/Source/ConcertSyncTest/Private/Replication/Components/ReplicationChangeStreamUtilsTests.cpp b/Engine/Plugins/Developer/Concert/ConcertSync/ConcertSyncTest/Source/ConcertSyncTest/Private/Replication/Components/ReplicationChangeStreamUtilsTests.cpp
--- a/Engine/Plugins/Developer/Concert/ConcertSync/ConcertSyncTest/Source/ConcertSyncTest/Private/Replication/Components/ReplicationChangeStreamUtilsTests.cpp
+++ b/Engine/Plugins/Developer/Concert/ConcertSync/ConcertSyncTest/Source/ConcertSyncTest/Private/Replication/Components/ReplicationChangeStreamUtilsTests.cpp
@@ -3,11 +3,11 @@
 #include "Replication/ChangeStreamSharedUtils.h"
 #include "Replication/Data/ConcertPropertySelection.h"
 #include "Replication/Data/ObjectReplicationMap.h"
-#include "Replication/Messages/ConcertReplicationEvents.h"
 
 #include "Misc/AutomationTest.h"
 #include "Components/SceneComponent.h"
 #include "Components/StaticMeshComponent.h"
+#include "Replication/Messages/ChangeStream.h"
 
 /**
  * Tests for UE::ConcertSyncCore::Replication::ChangeStreamUtils
@@ -52,7 +52,7 @@ namespace UE::ConcertSyncTests::ChangeStreamUtils
 		Desired.ReplicatedObjects.Add(ChangeProperties, { SceneComponentClass, Selection_RelativeLocationOnlyX });
 		Desired.ReplicatedObjects.Add(ChangeClass, { StaticMeshComponentClass, Selection_RelativeLocation });
 		Desired.ReplicatedObjects.Add(AddMe, { SceneComponentClass, Selection_RelativeLocation });
-		const FConcertChangeStream_Request Request = ConcertSyncCore::Replication::ChangeStreamUtils::BuildRequestFromDiff(StreamId, Base, Desired);
+		const FConcertReplication_ChangeStream_Request Request = ConcertSyncCore::Replication::ChangeStreamUtils::BuildRequestFromDiff(StreamId, Base, Desired);
 
 		
 		// 3. Test
@@ -64,7 +64,7 @@ namespace UE::ConcertSyncTests::ChangeStreamUtils
 		TestEqual(TEXT("Remove exactly 1"), Request.ObjectsToRemove.Num(), 1);
 		TestTrue(TEXT("Remove correct object"), Request.ObjectsToRemove.Contains({ StreamId, RemoveMe }));
 
-		if (const FConcertChangeStream_PutObject* ChangePropertiesChange = Request.ObjectsToPut.Find({ StreamId, ChangeProperties }))
+		if (const FConcertReplication_ChangeStream_PutObject* ChangePropertiesChange = Request.ObjectsToPut.Find({ StreamId, ChangeProperties }))
 		{
 			// No class change, so it should be null in the request
 			TestTrue(TEXT("Change Properties > Class unset"), ChangePropertiesChange->ClassPath.IsNull());
@@ -75,7 +75,7 @@ namespace UE::ConcertSyncTests::ChangeStreamUtils
 			AddError(TEXT("Did not change ChangeProperties object"));
 		}
 
-		if (const FConcertChangeStream_PutObject* ChangeClassChange = Request.ObjectsToPut.Find({ StreamId, ChangeClass }))
+		if (const FConcertReplication_ChangeStream_PutObject* ChangeClassChange = Request.ObjectsToPut.Find({ StreamId, ChangeClass }))
 		{
 			TestEqual(TEXT("Change Class > Class"), ChangeClassChange->ClassPath, StaticMeshComponentClass);
 			// No properties changed, so it should be empty in the request
@@ -86,7 +86,7 @@ namespace UE::ConcertSyncTests::ChangeStreamUtils
 			AddError(TEXT("Did not change ChangeClass object"));
 		}
 
-		if (const FConcertChangeStream_PutObject* AddMeChange = Request.ObjectsToPut.Find({ StreamId, AddMe }))
+		if (const FConcertReplication_ChangeStream_PutObject* AddMeChange = Request.ObjectsToPut.Find({ StreamId, AddMe }))
 		{
 			TestEqual(TEXT("Add Object > Class"), AddMeChange->ClassPath, SceneComponentClass);
 			TestEqual(TEXT("Add Object > Properties"), AddMeChange->Properties, Selection_RelativeLocation);
diff --git a/Engine/Plugins/Developer/Concert/ConcertSync/ConcertSyncTest/Source/ConcertSyncTest/Private/Replication/ReplicationAuthorityTests.cpp b/Engine/Plugins/Developer/Concert/ConcertSync/ConcertSyncTest/Source/ConcertSyncTest/Private/Replication/ReplicationAuthorityTests.cpp
--- a/Engine/Plugins/Developer/Concert/ConcertSync/ConcertSyncTest/Source/ConcertSyncTest/Private/Replication/ReplicationAuthorityTests.cpp
+++ b/Engine/Plugins/Developer/Concert/ConcertSync/ConcertSyncTest/Source/ConcertSyncTest/Private/Replication/ReplicationAuthorityTests.cpp
@@ -2,16 +2,17 @@
 
 #include "Util/SendReceiveObjectTestBase.h"
 
+#include "Replication/Formats/FullObjectFormat.h"
 #include "Replication/IConcertClientReplicationManager.h"
+#include "Replication/Messages/ObjectReplication.h"
 #include "TestReflectionObject.h"
+#include "Util/ConcertClientReplicationBridgeMock.h"
+#include "Util/SendReceiveGenericStreamTestBase.h"
 
 #include "Components/StaticMeshComponent.h"
 #include "Misc/AutomationTest.h"
 #include "Misc/Paths.h"
 #include "HAL/FileManager.h"
-#include "Replication/Formats/FullObjectFormat.h"
-#include "Util/ConcertClientReplicationBridgeMock.h"
-#include "Util/SendReceiveGenericStreamTestBase.h"
 
 namespace UE::ConcertSyncTests::Replication::Authority
 {
@@ -34,18 +35,18 @@ namespace UE::ConcertSyncTests::Replication::Authority
 			return false;
 		}
 		
-		FConcertObjectReplicationEvent ObjectReplicationEvent;
+		FConcertReplication_ObjectReplicationEvent ObjectReplicationEvent;
 		ObjectReplicationEvent.ReplicatedObject = TestObject;
 		ObjectReplicationEvent.SerializedPayload = *Payload;
-		FConcertStreamReplicationEvent ReplicationEvent;
+		FConcertReplication_StreamReplicationEvent ReplicationEvent;
 		ReplicationEvent.StreamId = SenderStreamId;
 		ReplicationEvent.ReplicatedObjects.Add(ObjectReplicationEvent);
-		FConcertBatchReplicationEvent ReplicationBatchEvent;
+		FConcertReplication_BatchReplicationEvent ReplicationBatchEvent;
 		ReplicationBatchEvent.Streams.Add(ReplicationEvent);
 		
 		// 2.2 Prepare for sending data without authority > Reject
 		bool bHasServerReceivedData = false;
-		auto OnServerReceive = [this, &bHasServerReceivedData](const FConcertSessionContext& Context, const FConcertBatchReplicationEvent& Event) mutable
+		auto OnServerReceive = [this, &bHasServerReceivedData](const FConcertSessionContext& Context, const FConcertReplication_BatchReplicationEvent& Event) mutable
 		{
 			if (bHasServerReceivedData)
 			{
@@ -53,12 +54,12 @@ namespace UE::ConcertSyncTests::Replication::Authority
 			}
 			bHasServerReceivedData = true;
 		};
-		auto OnClientReceive = [this](const FConcertSessionContext& Context, const FConcertBatchReplicationEvent& Event) mutable
+		auto OnClientReceive = [this](const FConcertSessionContext& Context, const FConcertReplication_BatchReplicationEvent& Event) mutable
 		{
 			AddError(TEXT("Server sent data from non-authorative client to receiving client!"));
 		};
-		ServerSession->RegisterCustomEventHandler<FConcertBatchReplicationEvent>(OnServerReceive);
-		Client_Receiver->ClientSessionMock->RegisterCustomEventHandler<FConcertBatchReplicationEvent>(OnClientReceive);
+		ServerSession->RegisterCustomEventHandler<FConcertReplication_BatchReplicationEvent>(OnServerReceive);
+		Client_Receiver->ClientSessionMock->RegisterCustomEventHandler<FConcertReplication_BatchReplicationEvent>(OnClientReceive);
 		const FGuid ServerSessionId { 0, 0, 0, 0};
 		
 		// 3. Test that server processed the data and rejected it.
diff --git a/Engine/Plugins/Developer/Concert/ConcertSync/ConcertSyncTest/Source/ConcertSyncTest/Private/Replication/ReplicationHandshakeTests.cpp b/Engine/Plugins/Developer/Concert/ConcertSync/ConcertSyncTest/Source/ConcertSyncTest/Private/Replication/ReplicationHandshakeTests.cpp
--- a/Engine/Plugins/Developer/Concert/ConcertSync/ConcertSyncTest/Source/ConcertSyncTest/Private/Replication/ReplicationHandshakeTests.cpp
+++ b/Engine/Plugins/Developer/Concert/ConcertSync/ConcertSyncTest/Source/ConcertSyncTest/Private/Replication/ReplicationHandshakeTests.cpp
@@ -7,7 +7,7 @@
 #include "Replication/IConcertClientReplicationBridge.h"
 #include "Replication/IConcertClientReplicationManager.h"
 #include "Replication/IConcertServerReplicationManager.h"
-#include "Replication/Messages/ConcertReplicationHandshakeMessages.h"
+#include "Replication/Messages/Handshake.h"
 #include "ReplicationTestInterface.h"
 #include "Util/ConcertClientReplicationBridgeMock.h"
 
diff --git a/Engine/Plugins/Developer/Concert/ConcertSync/ConcertSyncTest/Source/ConcertSyncTest/Private/Replication/ReplicationSendReceiveTests.cpp b/Engine/Plugins/Developer/Concert/ConcertSync/ConcertSyncTest/Source/ConcertSyncTest/Private/Replication/ReplicationSendReceiveTests.cpp
--- a/Engine/Plugins/Developer/Concert/ConcertSync/ConcertSyncTest/Source/ConcertSyncTest/Private/Replication/ReplicationSendReceiveTests.cpp
+++ b/Engine/Plugins/Developer/Concert/ConcertSync/ConcertSyncTest/Source/ConcertSyncTest/Private/Replication/ReplicationSendReceiveTests.cpp
@@ -5,8 +5,8 @@
 #include "Replication/Data/ReplicationStreamDescription.h"
 #include "Replication/IConcertClientReplicationManager.h"
 #include "Replication/IConcertServerReplicationManager.h"
-#include "Replication/Messages/ConcertReplicationEvents.h"
-#include "Replication/Messages/ConcertReplicationHandshakeMessages.h"
+#include "Replication/Messages/ObjectReplication.h"
+#include "Replication/Messages/Handshake.h"
 #include "Replication/PropertyChainUtils.h"
 #include "Replication/ReplicationTestInterface.h"
 #include "TestReflectionObject.h"
@@ -41,7 +41,7 @@ namespace UE::ConcertSyncTests::Replication::SendReceiveFlow
 	{
 		bool bHasServerReceivedData = false;
 		bool bHasClientReceivedData = false;
-		auto OnServerReceive = [&Test, &bHasServerReceivedData](const FConcertSessionContext& Context, const FConcertBatchReplicationEvent& Event) mutable
+		auto OnServerReceive = [&Test, &bHasServerReceivedData](const FConcertSessionContext& Context, const FConcertReplication_BatchReplicationEvent& Event) mutable
 		{
 			if (bHasServerReceivedData)
 			{
@@ -49,7 +49,7 @@ namespace UE::ConcertSyncTests::Replication::SendReceiveFlow
 			}
 			bHasServerReceivedData = true;
 		};
-		auto OnClientReceive = [&Test, &bHasClientReceivedData](const FConcertSessionContext& Context, const FConcertBatchReplicationEvent& Event) mutable
+		auto OnClientReceive = [&Test, &bHasClientReceivedData](const FConcertSessionContext& Context, const FConcertReplication_BatchReplicationEvent& Event) mutable
 		{
 			if (bHasClientReceivedData)
 			{
@@ -76,7 +76,7 @@ namespace UE::ConcertSyncTests::Replication::SendReceiveFlow
 		VectorOnlySelection.ReplicatedProperties.Add(*VectorPropertyChain);
 		
 		ConcertSyncClient::Replication::FChangeStreamRequest Request;
-		Request.ObjectsToPut.Add(FObjectInStreamID{ Test.SenderStreamId, Test.TestObject }, FConcertChangeStream_PutObject{ VectorOnlySelection });
+		Request.ObjectsToPut.Add(FObjectInStreamID{ Test.SenderStreamId, Test.TestObject }, FConcertReplication_ChangeStream_PutObject{ VectorOnlySelection });
 		bool bReceivedChangeStreamResponse = false;
 		Test.ClientReplicationManager_Sender->ChangeStream(Request)
 			.Next([&Test, &bReceivedChangeStreamResponse](ConcertSyncClient::Replication::FChangeStreamResponse&& Response)
diff --git a/Engine/Plugins/Developer/Concert/ConcertSync/ConcertSyncTest/Source/ConcertSyncTest/Private/Replication/ReplicationStreamRequestTests.cpp b/Engine/Plugins/Developer/Concert/ConcertSync/ConcertSyncTest/Source/ConcertSyncTest/Private/Replication/ReplicationStreamRequestTests.cpp
--- a/Engine/Plugins/Developer/Concert/ConcertSync/ConcertSyncTest/Source/ConcertSyncTest/Private/Replication/ReplicationStreamRequestTests.cpp
+++ b/Engine/Plugins/Developer/Concert/ConcertSync/ConcertSyncTest/Source/ConcertSyncTest/Private/Replication/ReplicationStreamRequestTests.cpp
@@ -339,7 +339,7 @@ namespace UE::ConcertSyncTests::Replication::Stream
 		FChangeStreamRequest AppendFloatRequest;
 		FConcertPropertySelection& NewSelection = GetPropertySelection(VectorFloatStream, *TestObject);
 		AddFloatProperty(NewSelection);
-		AppendFloatRequest.ObjectsToPut.Add(FObjectInStreamID{ VectorFloatStreamID, TestObject }, FConcertChangeStream_PutObject{ NewSelection });
+		AppendFloatRequest.ObjectsToPut.Add(FObjectInStreamID{ VectorFloatStreamID, TestObject }, FConcertReplication_ChangeStream_PutObject{ NewSelection });
 		ChangeStreamForSenderClientAndValidate(TEXT("AppendFloatRequest"), AppendFloatRequest, { FloatStream.BaseDescription, VectorFloatStream.BaseDescription });
 		
 		return true;
@@ -450,7 +450,7 @@ namespace UE::ConcertSyncTests::Replication::Stream
 		FChangeStreamRequest InvalidRequest;
 		FConcertPropertySelection NewSelection = GetPropertySelection(Stream, *TestObject);
 		AddFloatProperty(NewSelection);
-		InvalidRequest.ObjectsToPut.Add(FObjectInStreamID{ StreamId, TestObject}, FConcertChangeStream_PutObject{ NewSelection });
+		InvalidRequest.ObjectsToPut.Add(FObjectInStreamID{ StreamId, TestObject}, FConcertReplication_ChangeStream_PutObject{ NewSelection });
 		InvalidRequest.StreamsToAdd.Add(Stream.Pack()); // This will make it fail due to pre-existing stream ID
 
 		// Server logs a warning when rejecting - avoid the test being marked with a warning.
diff --git a/Engine/Plugins/Developer/Concert/ConcertSync/ConcertSyncTest/Source/ConcertSyncTest/Private/Replication/Util/SendReceiveGenericStreamTestBase.h b/Engine/Plugins/Developer/Concert/ConcertSync/ConcertSyncTest/Source/ConcertSyncTest/Private/Replication/Util/SendReceiveGenericStreamTestBase.h
--- a/Engine/Plugins/Developer/Concert/ConcertSync/ConcertSyncTest/Source/ConcertSyncTest/Private/Replication/Util/SendReceiveGenericStreamTestBase.h
+++ b/Engine/Plugins/Developer/Concert/ConcertSync/ConcertSyncTest/Source/ConcertSyncTest/Private/Replication/Util/SendReceiveGenericStreamTestBase.h
@@ -4,7 +4,7 @@
 
 #include "SendReceiveTestBase.h"
 
-struct FConcertBatchReplicationEvent;
+struct FConcertReplication_BatchReplicationEvent;
 class IConcertClientReplicationManager;
 class UTestReflectionObject;
 
diff --git a/Engine/Plugins/Developer/Concert/ConcertSync/ConcertSyncTest/Source/ConcertSyncTest/Private/Replication/Util/SendReceiveObjectTestBase.cpp b/Engine/Plugins/Developer/Concert/ConcertSync/ConcertSyncTest/Source/ConcertSyncTest/Private/Replication/Util/SendReceiveObjectTestBase.cpp
--- a/Engine/Plugins/Developer/Concert/ConcertSync/ConcertSyncTest/Source/ConcertSyncTest/Private/Replication/Util/SendReceiveObjectTestBase.cpp
+++ b/Engine/Plugins/Developer/Concert/ConcertSync/ConcertSyncTest/Source/ConcertSyncTest/Private/Replication/Util/SendReceiveObjectTestBase.cpp
@@ -6,8 +6,8 @@
 #include "Replication/Data/ReplicationStreamDescription.h"
 #include "Replication/IConcertClientReplicationManager.h"
 #include "Replication/IConcertServerReplicationManager.h"
-#include "Replication/Messages/ConcertReplicationEvents.h"
-#include "Replication/Messages/ConcertReplicationHandshakeMessages.h"
+#include "Replication/Messages/ObjectReplication.h"
+#include "Replication/Messages/Handshake.h"
 #include "Replication/PropertyChainUtils.h"
 #include "Replication/ReplicationTestInterface.h"
 #include "Replication/TestReflectionObject.h"
@@ -51,7 +51,7 @@ namespace UE::ConcertSyncTests::Replication
 		EPropertyTestFlags PropertyFlags
 		)
 	{
-		auto TestReplicationData_Server = [this, OnServerReceive](const FConcertSessionContext& Context, const FConcertBatchReplicationEvent& Event)
+		auto TestReplicationData_Server = [this, OnServerReceive](const FConcertSessionContext& Context, const FConcertReplication_BatchReplicationEvent& Event)
 		{
 			const TSet<FGuid> SenderStreamIds = GetSenderStreamIds();
 			TestEqual(TEXT("Server received right number of streams"), Event.Streams.Num(), SenderStreamIds.Num());
@@ -64,7 +64,7 @@ namespace UE::ConcertSyncTests::Replication
 			}
 			OnServerReceive(Context, Event);
 		};
-		auto TestReplicationData_Client_Receiver = [this, OnReceiverClientReceive](const FConcertSessionContext& Context, const FConcertBatchReplicationEvent& Event)
+		auto TestReplicationData_Client_Receiver = [this, OnReceiverClientReceive](const FConcertSessionContext& Context, const FConcertReplication_BatchReplicationEvent& Event)
 		{
 			const TSet<FGuid> SenderStreamIds = GetSenderStreamIds();
 			TestEqual(TEXT("Client 2 received right number of streams"), Event.Streams.Num(), SenderStreamIds.Num());
@@ -77,8 +77,8 @@ namespace UE::ConcertSyncTests::Replication
 			}
 			OnReceiverClientReceive(Context, Event);
 		};
-		const FDelegateHandle ServerHandle = ServerSession->RegisterCustomEventHandler<FConcertBatchReplicationEvent>(TestReplicationData_Server);
-		const FDelegateHandle ClientHandle = Client_Receiver->ClientSessionMock->RegisterCustomEventHandler<FConcertBatchReplicationEvent>(TestReplicationData_Client_Receiver);
+		const FDelegateHandle ServerHandle = ServerSession->RegisterCustomEventHandler<FConcertReplication_BatchReplicationEvent>(TestReplicationData_Server);
+		const FDelegateHandle ClientHandle = Client_Receiver->ClientSessionMock->RegisterCustomEventHandler<FConcertReplication_BatchReplicationEvent>(TestReplicationData_Client_Receiver);
 
 		
 		// TestObject is the same UObject on both clients.
@@ -94,8 +94,8 @@ namespace UE::ConcertSyncTests::Replication
 		SetDifferentValues(*TestObject, PropertyFlags);
 		TickClient(Client_Receiver);
 
-		ServerSession->UnregisterCustomEventHandler<FConcertBatchReplicationEvent>(ServerHandle);
-		Client_Receiver->ClientSessionMock->UnregisterCustomEventHandler<FConcertBatchReplicationEvent>(ClientHandle);
+		ServerSession->UnregisterCustomEventHandler<FConcertReplication_BatchReplicationEvent>(ServerHandle);
+		Client_Receiver->ClientSessionMock->UnregisterCustomEventHandler<FConcertReplication_BatchReplicationEvent>(ClientHandle);
 		
 		// No call to Super because we're completely overriding the behavior.
 	}
diff --git a/Engine/Plugins/Developer/Concert/ConcertSync/ConcertSyncTest/Source/ConcertSyncTest/Private/Replication/Util/SendReceiveObjectTestBase.h b/Engine/Plugins/Developer/Concert/ConcertSync/ConcertSyncTest/Source/ConcertSyncTest/Private/Replication/Util/SendReceiveObjectTestBase.h
--- a/Engine/Plugins/Developer/Concert/ConcertSync/ConcertSyncTest/Source/ConcertSyncTest/Private/Replication/Util/SendReceiveObjectTestBase.h
+++ b/Engine/Plugins/Developer/Concert/ConcertSync/ConcertSyncTest/Source/ConcertSyncTest/Private/Replication/Util/SendReceiveObjectTestBase.h
@@ -4,7 +4,7 @@
 
 #include "SendReceiveTestBase.h"
 
-struct FConcertBatchReplicationEvent;
+struct FConcertReplication_BatchReplicationEvent;
 class IConcertClientReplicationManager;
 class UTestReflectionObject;
 
diff --git a/Engine/Plugins/Developer/Concert/ConcertSync/ConcertSyncTest/Source/ConcertSyncTest/Private/Replication/Util/SendReceiveTestBase.cpp b/Engine/Plugins/Developer/Concert/ConcertSync/ConcertSyncTest/Source/ConcertSyncTest/Private/Replication/Util/SendReceiveTestBase.cpp
--- a/Engine/Plugins/Developer/Concert/ConcertSync/ConcertSyncTest/Source/ConcertSyncTest/Private/Replication/Util/SendReceiveTestBase.cpp
+++ b/Engine/Plugins/Developer/Concert/ConcertSync/ConcertSyncTest/Source/ConcertSyncTest/Private/Replication/Util/SendReceiveTestBase.cpp
@@ -4,8 +4,8 @@
 
 #include "ConcertClientReplicationBridgeMock.h"
 #include "Replication/IConcertClientReplicationManager.h"
-#include "Replication/Messages/ConcertReplicationEvents.h"
-#include "Replication/Messages/ConcertReplicationHandshakeMessages.h"
+#include "Replication/Messages/ObjectReplication.h"
+#include "Replication/Messages/Handshake.h"
 #include "Replication/ReplicationTestInterface.h"
 #include "Util/ClientServerCommunicationTest.h"
 
@@ -67,18 +67,18 @@ namespace UE::ConcertSyncTests::Replication
 		}
 		
 		// 1.3 Test basic assumptions: only sender sends, only receiver receives.
-		auto TestReplicationData_Server = [&](const FConcertSessionContext& Context, const FConcertBatchReplicationEvent& Event)
+		auto TestReplicationData_Server = [&](const FConcertSessionContext& Context, const FConcertReplication_BatchReplicationEvent& Event)
 		{
 			TestEqual(TEXT("Only sender is supposed to send data!"), Context.SourceEndpointId, Client_Sender->ClientSessionMock->GetSessionClientEndpointId());
 		};
-		auto TestReplicationData_Client_Sender = [&](const FConcertSessionContext& Context, const FConcertBatchReplicationEvent& Event)
+		auto TestReplicationData_Client_Sender = [&](const FConcertSessionContext& Context, const FConcertReplication_BatchReplicationEvent& Event)
         {
         	// There was a bug where client would receive its own data...
 			AddError(TEXT("Sender was not supposed to receive data!"));
         };
 		
-		ServerSession->RegisterCustomEventHandler<FConcertBatchReplicationEvent>(TestReplicationData_Server);
-		Client_Sender->ClientSessionMock->RegisterCustomEventHandler<FConcertBatchReplicationEvent>(TestReplicationData_Client_Sender);
+		ServerSession->RegisterCustomEventHandler<FConcertReplication_BatchReplicationEvent>(TestReplicationData_Server);
+		Client_Sender->ClientSessionMock->RegisterCustomEventHandler<FConcertReplication_BatchReplicationEvent>(TestReplicationData_Client_Sender);
 	}
 	
 	void FSendReceiveTestBase::SimulateSenderToReceiver(
@@ -86,16 +86,16 @@ namespace UE::ConcertSyncTests::Replication
 		TFunctionRef<FReceiveReplicationEventSignature> OnReceiverClientReceive
 		)
 	{
-		auto TestReplicationData_Server = [this, OnServerReceive](const FConcertSessionContext& Context, const FConcertBatchReplicationEvent& Event)
+		auto TestReplicationData_Server = [this, OnServerReceive](const FConcertSessionContext& Context, const FConcertReplication_BatchReplicationEvent& Event)
 		{
 			OnServerReceive(Context, Event);
 		};
-		auto TestReplicationData_Client_Receiver = [this, OnReceiverClientReceive](const FConcertSessionContext& Context, const FConcertBatchReplicationEvent& Event)
+		auto TestReplicationData_Client_Receiver = [this, OnReceiverClientReceive](const FConcertSessionContext& Context, const FConcertReplication_BatchReplicationEvent& Event)
 		{
 			OnReceiverClientReceive(Context, Event);
 		};
-		const FDelegateHandle ServerHandle = ServerSession->RegisterCustomEventHandler<FConcertBatchReplicationEvent>(TestReplicationData_Server);
-		const FDelegateHandle ClientHandle = Client_Receiver->ClientSessionMock->RegisterCustomEventHandler<FConcertBatchReplicationEvent>(TestReplicationData_Client_Receiver);
+		const FDelegateHandle ServerHandle = ServerSession->RegisterCustomEventHandler<FConcertReplication_BatchReplicationEvent>(TestReplicationData_Server);
+		const FDelegateHandle ClientHandle = Client_Receiver->ClientSessionMock->RegisterCustomEventHandler<FConcertReplication_BatchReplicationEvent>(TestReplicationData_Client_Receiver);
 
 		// 1. Sender > Server
 		TickClient(Client_Sender);
@@ -104,8 +104,8 @@ namespace UE::ConcertSyncTests::Replication
 		// 3. Receive from server
 		TickClient(Client_Receiver);
 		
-		ServerSession->UnregisterCustomEventHandler<FConcertBatchReplicationEvent>(ServerHandle);
-		Client_Receiver->ClientSessionMock->UnregisterCustomEventHandler<FConcertBatchReplicationEvent>(ClientHandle);
+		ServerSession->UnregisterCustomEventHandler<FConcertReplication_BatchReplicationEvent>(ServerHandle);
+		Client_Receiver->ClientSessionMock->UnregisterCustomEventHandler<FConcertReplication_BatchReplicationEvent>(ClientHandle);
 	}
 	
 	void FSendReceiveTestBase::TickClient(FClientInfo* Client)
diff --git a/Engine/Plugins/Developer/Concert/ConcertSync/ConcertSyncTest/Source/ConcertSyncTest/Private/Replication/Util/SendReceiveTestBase.h b/Engine/Plugins/Developer/Concert/ConcertSync/ConcertSyncTest/Source/ConcertSyncTest/Private/Replication/Util/SendReceiveTestBase.h
--- a/Engine/Plugins/Developer/Concert/ConcertSync/ConcertSyncTest/Source/ConcertSyncTest/Private/Replication/Util/SendReceiveTestBase.h
+++ b/Engine/Plugins/Developer/Concert/ConcertSync/ConcertSyncTest/Source/ConcertSyncTest/Private/Replication/Util/SendReceiveTestBase.h
@@ -3,6 +3,7 @@
 #pragma once
 
 #include "Replication/IConcertClientReplicationManager.h"
+#include "Replication/Messages/ObjectReplication.h"
 #include "Util/ClientServerCommunicationTest.h"
 
 class IConcertClientReplicationManager;
@@ -30,7 +31,7 @@ namespace UE::ConcertSyncTests::Replication
 		
 		static ConcertSyncClient::Replication::FJoinReplicatedSessionArgs CreateHandshakeArgsFrom(const UObject& Object, const FGuid& SenderStreamId = FGuid::NewGuid());
 		
-		using FReceiveReplicationEventSignature = void(const FConcertSessionContext& Context, const FConcertBatchReplicationEvent& Event);
+		using FReceiveReplicationEventSignature = void(const FConcertSessionContext& Context, const FConcertReplication_BatchReplicationEvent& Event);
 		
 		float FakeDeltaTime = 1.f / 60.f;
 		
