Predict client conflicts and remove them from requests when submitting them to the server

diff --git a/Engine/Plugins/Developer/Concert/ConcertApp/MultiUserClient/Source/MultiUserClient/Private/Replication/Client/LocalReplicationClient.cpp b/Engine/Plugins/Developer/Concert/ConcertApp/MultiUserClient/Source/MultiUserClient/Private/Replication/Client/LocalReplicationClient.cpp
--- a/Engine/Plugins/Developer/Concert/ConcertApp/MultiUserClient/Source/MultiUserClient/Private/Replication/Client/LocalReplicationClient.cpp
+++ b/Engine/Plugins/Developer/Concert/ConcertApp/MultiUserClient/Source/MultiUserClient/Private/Replication/Client/LocalReplicationClient.cpp
@@ -11,21 +11,22 @@ namespace UE::MultiUserClient
 	FLocalReplicationClient::FLocalReplicationClient(
 		UMultiUserReplicationClientPreset& InSessionContent,
 		TUniquePtr<IClientStreamSynchronizer> InStreamSynchronizer,
-		TSharedRef<IConcertSyncClient> InClient
+		TSharedRef<IConcertSyncClient> InClient,
+		const FGlobalAuthorityCache& InAuthorityCache
 		)
 		: FReplicationClient(
 			InClient->GetConcertClient()->GetCurrentSession()->GetSessionClientEndpointId(),
 			InSessionContent,
 			MoveTemp(InStreamSynchronizer),
-			MakeUnique<FAuthoritySynchronizer_LocalClient>(
-				InClient,
-				FDoesObjectHaveProperties::CreateLambda([this](const FSoftObjectPath& ObjectPath)
-				{
-					return GetStreamDiffer().DoesObjectHavePropertiesAfterSubmit(ObjectPath);
-				})),
-			[this, InClient](FStreamChangeTracker& InStreamChangeTracker, FAuthorityChangeTracker& InAuthorityChangeTracker, IClientStreamSynchronizer& InStreamSynchronizer)
+		MakeUnique<FAuthoritySynchronizer_LocalClient>(
+			InClient,
+			FDoesObjectHaveProperties::CreateLambda([this](const FSoftObjectPath& ObjectPath)
 			{
-				return MakeUnique<FSubmissionWorkflow_LocalClient>(InClient, InStreamChangeTracker, InAuthorityChangeTracker, InStreamSynchronizer);
+				return GetStreamDiffer().DoesObjectHavePropertiesAfterSubmit(ObjectPath);
+			})),
+			[this, InClient, &InAuthorityCache](FStreamChangeTracker& InStreamChangeTracker, FAuthorityChangeTracker& InAuthorityChangeTracker, IClientStreamSynchronizer& InStreamSynchronizer) mutable
+			{
+				return MakeUnique<FSubmissionWorkflow_LocalClient>(MoveTemp(InClient), InStreamChangeTracker, InAuthorityChangeTracker, InStreamSynchronizer, InAuthorityCache);
 			})
 	{}
 }
diff --git a/Engine/Plugins/Developer/Concert/ConcertApp/MultiUserClient/Source/MultiUserClient/Private/Replication/Client/LocalReplicationClient.h b/Engine/Plugins/Developer/Concert/ConcertApp/MultiUserClient/Source/MultiUserClient/Private/Replication/Client/LocalReplicationClient.h
--- a/Engine/Plugins/Developer/Concert/ConcertApp/MultiUserClient/Source/MultiUserClient/Private/Replication/Client/LocalReplicationClient.h
+++ b/Engine/Plugins/Developer/Concert/ConcertApp/MultiUserClient/Source/MultiUserClient/Private/Replication/Client/LocalReplicationClient.h
@@ -6,6 +6,8 @@
 
 namespace UE::MultiUserClient
 {
+	class FGlobalAuthorityCache;
+
 	/** Holds extra information about a local replication client. */
 	class FLocalReplicationClient : public FReplicationClient
 	{
@@ -14,7 +16,8 @@ namespace UE::MultiUserClient
 		FLocalReplicationClient(
 			UMultiUserReplicationClientPreset& InSessionContent,
 			TUniquePtr<IClientStreamSynchronizer> InStreamSynchronizer,
-			TSharedRef<IConcertSyncClient> InClient
+			TSharedRef<IConcertSyncClient> InClient,
+			const FGlobalAuthorityCache& InAuthorityCache
 			);
 	};
 }
diff --git a/Engine/Plugins/Developer/Concert/ConcertApp/MultiUserClient/Source/MultiUserClient/Private/Replication/Client/ReplicationClientManager.cpp b/Engine/Plugins/Developer/Concert/ConcertApp/MultiUserClient/Source/MultiUserClient/Private/Replication/Client/ReplicationClientManager.cpp
--- a/Engine/Plugins/Developer/Concert/ConcertApp/MultiUserClient/Source/MultiUserClient/Private/Replication/Client/ReplicationClientManager.cpp
+++ b/Engine/Plugins/Developer/Concert/ConcertApp/MultiUserClient/Source/MultiUserClient/Private/Replication/Client/ReplicationClientManager.cpp
@@ -19,15 +19,16 @@ namespace UE::MultiUserClient
 		)
 		: SessionContent(NewObject<UMultiUserReplicationSessionPreset>(GetTransientPackage(), NAME_None, RF_Transient))
 		, Session(InSession)
+		, QueryService(InClient)
+		, AuthorityCache(*this)
 		, LocalClient([this, InClient]()
 		{
 			UMultiUserReplicationClientPreset* ClientPreset = SessionContent->AddClient();
-			return FLocalReplicationClient(*ClientPreset, MakeUnique<FStreamSynchronizer_LocalClient>(InClient, ClientPreset->Stream->StreamId), InClient);
+			return FLocalReplicationClient(*ClientPreset, MakeUnique<FStreamSynchronizer_LocalClient>(InClient, ClientPreset->Stream->StreamId), InClient, AuthorityCache);
 		}())
-		, QueryService(InClient)
-		, AuthorityCache(*this)
 		, SubmissionNotifier(*this)
 	{
+		AuthorityCache.RegisterEvents();
 		InSession->OnSessionClientChanged().AddRaw(this, &FReplicationClientManager::OnSessionClientChanged);
 
 		for (const FGuid& ClientEndpointId : InSession->GetSessionClientEndpointIds())
diff --git a/Engine/Plugins/Developer/Concert/ConcertApp/MultiUserClient/Source/MultiUserClient/Private/Replication/Submission/SubmissionWorkflow_LocalClient.cpp b/Engine/Plugins/Developer/Concert/ConcertApp/MultiUserClient/Source/MultiUserClient/Private/Replication/Submission/SubmissionWorkflow_LocalClient.cpp
--- a/Engine/Plugins/Developer/Concert/ConcertApp/MultiUserClient/Source/MultiUserClient/Private/Replication/Submission/SubmissionWorkflow_LocalClient.cpp
+++ b/Engine/Plugins/Developer/Concert/ConcertApp/MultiUserClient/Source/MultiUserClient/Private/Replication/Submission/SubmissionWorkflow_LocalClient.cpp
@@ -5,6 +5,7 @@
 #include "IConcertSyncClient.h"
 #include "Replication/Authority/AuthorityChangeTracker.h"
 #include "Replication/Stream/StreamChangeTracker.h"
+#include "Replication/Util/GlobalAuthorityCache.h"
 #include "Replication/Util/StreamRequestUtils.h"
 
 namespace UE::MultiUserClient
@@ -13,9 +14,11 @@ namespace UE::MultiUserClient
 		TSharedRef<IConcertSyncClient> InClient,
 		FStreamChangeTracker& InStreamChangeTracker,
 		FAuthorityChangeTracker& InAuthorityChangeTracker,
-		IClientStreamSynchronizer& InStreamSynchronizer
+		IClientStreamSynchronizer& InStreamSynchronizer,
+		const FGlobalAuthorityCache& InAuthorityCache
 		)
 		: Client(MoveTemp(InClient))
+		, AuthorityCache(InAuthorityCache)
 		, StreamChangeTracker(InStreamChangeTracker)
 		, AuthorityChangeTracker(InAuthorityChangeTracker)
 		, StreamSynchronizer(InStreamSynchronizer)
@@ -55,9 +58,12 @@ namespace UE::MultiUserClient
 			return nullptr;
 		}
 
-		// TODO DP: We should check the authority change for conflicts one more time here, in case the remote client's streams have changed since the user last edited authority
+		const FGuid LocalClientId = Client->GetConcertClient()->GetCurrentSession()->GetSessionClientEndpointId();
+		
 		// The authority request is pre-built now to avoid sending changes the local client makes while we're waiting for the latent server responses
 		FAuthorityChangeRequest AuthorityChangeRequest = AuthorityChangeTracker.BuildChangeRequest(GetLocalClientStreamId());
+		// Predict conflicts in case the remote client's streams have changed. Also: while the UI highlights "bad" requests, it does not correct it.
+		AuthorityCache.CleanseConflictsFrom(AuthorityChangeRequest, LocalClientId);
 		
 		const FStreamChangelist& Changelist = StreamChangeTracker.GetCachedDeltaChange();
 		const bool bIsChangelistEmpty = Changelist.ObjectsToPut.IsEmpty() && Changelist.ObjectsToRemove.IsEmpty();
@@ -79,6 +85,8 @@ namespace UE::MultiUserClient
 			FChangeStreamRequest Request = StreamSynchronizer.GetServerState().ReplicatedObjects.IsEmpty()
 				? StreamRequestUtils::BuildChangeRequest_CreateNewStream(GetLocalClientStreamId(), Changelist)
 				: StreamRequestUtils::BuildChangeRequest_UpdateExistingStream(Changelist);
+			// Predict conflicts in case the remote client's streams have changed. Also: while the UI highlights "bad" requests, it does not correct it.
+			AuthorityCache.CleanseConflictsFrom(AuthorityChangeRequest, LocalClientId);
 		
 			ReplicationManager->ChangeStream(Request)
 				.Next([this, DestructionDetection = LifetimeToken->AsWeak(), Request, AuthorityChangeRequest = MoveTemp(AuthorityChangeRequest)](FChangeStreamResponse&& Response)
diff --git a/Engine/Plugins/Developer/Concert/ConcertApp/MultiUserClient/Source/MultiUserClient/Private/Replication/Submission/SubmissionWorkflow_LocalClient.h b/Engine/Plugins/Developer/Concert/ConcertApp/MultiUserClient/Source/MultiUserClient/Private/Replication/Submission/SubmissionWorkflow_LocalClient.h
--- a/Engine/Plugins/Developer/Concert/ConcertApp/MultiUserClient/Source/MultiUserClient/Private/Replication/Submission/SubmissionWorkflow_LocalClient.h
+++ b/Engine/Plugins/Developer/Concert/ConcertApp/MultiUserClient/Source/MultiUserClient/Private/Replication/Submission/SubmissionWorkflow_LocalClient.h
@@ -17,6 +17,7 @@ namespace UE::MultiUserClient
 	class IClientAuthoritySynchronizer;
 	class IClientStreamSynchronizer;
 	class FAuthorityChangeTracker;
+	class FGlobalAuthorityCache;
 	class FStreamChangeTracker;
 	
 	/** Handles the submission workflow for a single client. */
@@ -28,7 +29,8 @@ namespace UE::MultiUserClient
 			TSharedRef<IConcertSyncClient> InClient,
 			FStreamChangeTracker& InStreamChangeTracker,
 			FAuthorityChangeTracker& InAuthorityChangeTracker,
-			IClientStreamSynchronizer& InStreamSynchronizer
+			IClientStreamSynchronizer& InStreamSynchronizer,
+			const FGlobalAuthorityCache& InAuthorityCache
 			);
 		
 		//~ Begin ISubmissionWorkflow Interface
@@ -45,6 +47,9 @@ namespace UE::MultiUserClient
 		/** Used to send authority requests to the server. */
 		const TSharedRef<IConcertSyncClient> Client;
 
+		/** Used to filter out changes that would cause a conflict when submitted. */
+		const FGlobalAuthorityCache& AuthorityCache;
+
 		/** Used to get changes made to the stream */
 		FStreamChangeTracker& StreamChangeTracker;
 		/** Used to get changes made to the authority */
diff --git a/Engine/Plugins/Developer/Concert/ConcertApp/MultiUserClient/Source/MultiUserClient/Private/Replication/Util/GlobalAuthorityCache.cpp b/Engine/Plugins/Developer/Concert/ConcertApp/MultiUserClient/Source/MultiUserClient/Private/Replication/Util/GlobalAuthorityCache.cpp
--- a/Engine/Plugins/Developer/Concert/ConcertApp/MultiUserClient/Source/MultiUserClient/Private/Replication/Util/GlobalAuthorityCache.cpp
+++ b/Engine/Plugins/Developer/Concert/ConcertApp/MultiUserClient/Source/MultiUserClient/Private/Replication/Util/GlobalAuthorityCache.cpp
@@ -10,6 +10,9 @@ namespace UE::MultiUserClient
 {
 	FGlobalAuthorityCache::FGlobalAuthorityCache(FReplicationClientManager& InClientManager)
 		: ClientManager(InClientManager)
+	{}
+
+	void FGlobalAuthorityCache::RegisterEvents()
 	{
 		RegisterForClientEvents(ClientManager.GetLocalClient());
 
@@ -17,19 +20,6 @@ namespace UE::MultiUserClient
 		ClientManager.OnPreRemoteClientRemoved().AddRaw(this, &FGlobalAuthorityCache::OnPreRemoteClientRemoved);
 	}
 
-	FGlobalAuthorityCache::~FGlobalAuthorityCache()
-	{
-		// FGlobalAuthorityCache is owned by FRelicationClientManager so this is strictly not needed. But we'll follow RAII in case ownership changes.
-		UnregisterFromClientEvents(ClientManager.GetLocalClient());
-		for (const TNonNullPtr<FRemoteReplicationClient>& RemoteClient : ClientManager.GetRemoteClients())
-		{
-			UnregisterFromClientEvents(*RemoteClient);
-		}
-		
-		ClientManager.OnPostRemoteClientAdded().RemoveAll(this);
-		ClientManager.OnPreRemoteClientRemoved().RemoveAll(this);
-	}
-
 	void FGlobalAuthorityCache::ForEachClientWithAuthorityOverObject(const FSoftObjectPath& Object, TFunctionRef<EBreakBehavior(const FGuid& ClientId)> Callback) const
 	{
 		const TSet<FGuid>* Clients = OwnedObjectsToClients.Find(Object);
@@ -129,6 +119,16 @@ namespace UE::MultiUserClient
 		return Result;
 	}
 
+	void FGlobalAuthorityCache::CleanseConflictsFrom(ConcertSyncClient::Replication::FAuthorityChangeRequest& Request, const FGuid& SendingClient) const
+	{
+		ConcertSyncCore::Replication::AuthorityConflictUtils::CleanseConflictsFrom(Request, SendingClient, *this);
+	}
+
+	void FGlobalAuthorityCache::CleanseConflictsFrom(ConcertSyncClient::Replication::FChangeStreamRequest& Request, const FGuid& SendingClient) const
+	{
+		ConcertSyncCore::Replication::AuthorityConflictUtils::CleanseConflictsFrom(Request, SendingClient, *this);
+	}
+
 	void FGlobalAuthorityCache::RegisterForClientEvents(const FReplicationClient& Client)
 	{
 		const FGuid& ClientEndpointId = Client.GetEndpointId();
@@ -173,7 +173,7 @@ namespace UE::MultiUserClient
 			}
 		}
 	}
-
+	
 	void FGlobalAuthorityCache::ForEachStream(const FGuid& ClientEndpointId, TFunctionRef<EBreakBehavior(const FGuid& StreamId, const FObjectReplicationMap& ReplicationMap)> Callback) const
 	{
 		const FReplicationClient* Client = ClientManager.FindClient(ClientEndpointId);
diff --git a/Engine/Plugins/Developer/Concert/ConcertApp/MultiUserClient/Source/MultiUserClient/Private/Replication/Util/GlobalAuthorityCache.h b/Engine/Plugins/Developer/Concert/ConcertApp/MultiUserClient/Source/MultiUserClient/Private/Replication/Util/GlobalAuthorityCache.h
--- a/Engine/Plugins/Developer/Concert/ConcertApp/MultiUserClient/Source/MultiUserClient/Private/Replication/Util/GlobalAuthorityCache.h
+++ b/Engine/Plugins/Developer/Concert/ConcertApp/MultiUserClient/Source/MultiUserClient/Private/Replication/Util/GlobalAuthorityCache.h
@@ -24,14 +24,17 @@ namespace UE::MultiUserClient
 	 * Allows efficient look-up of which objects and properties are owned by which clients.
 	 * This class efficiently answers the question: "Which clients own this object?"
 	 */
-	class FGlobalAuthorityCache : public FNoncopyable, private ConcertSyncCore::Replication::AuthorityConflictUtils::IReplicationGroundTruth
+	class FGlobalAuthorityCache
+		: public FNoncopyable
+		, private ConcertSyncCore::Replication::AuthorityConflictUtils::IReplicationGroundTruth
 	{
 	public:
 
 		using FProcessPropertyConflict = TFunctionRef<EBreakBehavior(const FGuid& ConflictingClientId, const FConcertPropertyChain& Property)>;
 		
 		FGlobalAuthorityCache(FReplicationClientManager& InClientManager);
-		~FGlobalAuthorityCache();
+		/** Called when the local client has been created and it is safe to register client events with FReplicationClientManager. */
+		void RegisterEvents();
 
 		/** Iterates every client that has authority over Object. */
 		void ForEachClientWithAuthorityOverObject(const FSoftObjectPath& Object, TFunctionRef<EBreakBehavior(const FGuid& ClientId)> Callback) const;
@@ -53,7 +56,22 @@ namespace UE::MultiUserClient
 
 		/** Gets the client that has authority over the given property, if there is any. */
 		TOptional<FGuid> GetClientWithAuthorityOverProperty(const FSoftObjectPath& Object, const FConcertPropertyChain& Property) const;
-
+		
+		/**
+		 * Removes entries from Request that would generate conflicts.
+		 * 
+		 * @param Request The request to clense
+		 * @param SendingClient The client that will send the request
+		 */
+		void CleanseConflictsFrom(ConcertSyncClient::Replication::FAuthorityChangeRequest& Request, const FGuid& SendingClient) const;
+		/**
+		 * Removes entries from Request that would generate conflicts.
+		 * 
+		 * @param Request The request to clense
+		 * @param SendingClient The client that will send the request
+		 */
+		void CleanseConflictsFrom(ConcertSyncClient::Replication::FChangeStreamRequest& Request, const FGuid& SendingClient) const;
+		
 		DECLARE_MULTICAST_DELEGATE_OneParam(FOnCacheChanged, const FGuid& ClientId);
 		/** Called when the cache changes for a specific client. */
 		FOnCacheChanged& OnCacheChanged() { return OnCacheChangedDelegate; }
diff --git a/Engine/Plugins/Developer/Concert/ConcertSync/ConcertSyncCore/Source/ConcertSyncCore/Private/Replication/AuthorityConflictSharedUtils.cpp b/Engine/Plugins/Developer/Concert/ConcertSync/ConcertSyncCore/Source/ConcertSyncCore/Private/Replication/AuthorityConflictSharedUtils.cpp
--- a/Engine/Plugins/Developer/Concert/ConcertSync/ConcertSyncCore/Source/ConcertSyncCore/Private/Replication/AuthorityConflictSharedUtils.cpp
+++ b/Engine/Plugins/Developer/Concert/ConcertSync/ConcertSyncCore/Source/ConcertSyncCore/Private/Replication/AuthorityConflictSharedUtils.cpp
@@ -3,6 +3,8 @@
 #include "Replication/AuthorityConflictSharedUtils.h"
 
 #include "Replication/Data/ReplicationStreamDescription.h"
+#include "Replication/Messages/ChangeAuthority.h"
+#include "Replication/Messages/ChangeStream.h"
 
 namespace UE::ConcertSyncCore::Replication::AuthorityConflictUtils
 {
@@ -81,4 +83,69 @@ namespace UE::ConcertSyncCore::Replication::AuthorityConflictUtils
 			);
 		return bFreeOfConflicts ? EAuthorityConflict::Allowed : EAuthorityConflict::Conflict;
 	}
+
+	void CleanseConflictsFrom(FConcertReplication_ChangeAuthority_Request& Request, const FGuid& SendingClient, const IReplicationGroundTruth& GroundTruth)
+	{
+		// Need to check whether TakeAuthority is taking authority over properties other clients are already replicating
+		GroundTruth.ForEachStream(SendingClient, [&SendingClient, &Request, &GroundTruth](const FGuid& StreamId, const FObjectReplicationMap& ReplicationMap)
+		{
+			for (auto ChangeIt = Request.TakeAuthority.CreateIterator(); ChangeIt; ++ChangeIt)
+			{
+				TPair<FSoftObjectPath, FConcertStreamArray>& Change = *ChangeIt;
+				
+				const FReplicatedObjectInfo* ObjectInfo = ReplicationMap.ReplicatedObjects.Find(Change.Key);
+				if (!ObjectInfo)
+				{
+					ChangeIt.RemoveCurrent();
+					continue;
+				}
+
+				const FSoftObjectPath& ObjectPath = Change.Key;
+				const TArray<FConcertPropertyChain>& Properties = ObjectInfo->PropertySelection.ReplicatedProperties;
+				const EAuthorityConflict Conflict = EnumerateAuthorityConflicts(SendingClient, ObjectPath, Properties, GroundTruth);
+
+				const bool bHasConflict = Conflict == EAuthorityConflict::Conflict;
+				FConcertStreamArray& StreamArray = Change.Value;
+				if (bHasConflict && StreamArray.StreamIds.Num() > 1)
+				{
+					StreamArray.StreamIds.Remove(StreamId);
+				}
+				else if (bHasConflict)
+				{
+					ChangeIt.RemoveCurrent();
+				}
+			}
+
+			return Request.TakeAuthority.IsEmpty() ? EBreakBehavior::Break : EBreakBehavior::Continue;
+		});
+	}
+	
+	void CleanseConflictsFrom(FConcertReplication_ChangeStream_Request& Request, const FGuid& SendingClient, const IReplicationGroundTruth& GroundTruth)
+	{
+		// Need to check whether ObjectsToPut adds any properties that an existing client has authority over.
+		GroundTruth.ForEachStream(SendingClient, [&SendingClient, &Request, &GroundTruth](const FGuid& StreamId, const FObjectReplicationMap& ReplicationMap)
+		{
+			for (auto ChangeIt = Request.ObjectsToPut.CreateIterator(); ChangeIt; ++ChangeIt)
+			{
+				const TPair<FObjectInStreamID, FConcertReplication_ChangeStream_PutObject>& Change = *ChangeIt;
+				const FConcertReplication_ChangeStream_PutObject& PutObject = Change.Value;
+
+				// Just changing the class?
+				if (PutObject.Properties.ReplicatedProperties.IsEmpty())
+				{
+					continue;
+				}
+				
+				const FSoftObjectPath& ObjectPath = Change.Key.Object;
+				const EAuthorityConflict Conflict = EnumerateAuthorityConflicts(SendingClient, ObjectPath, PutObject.Properties.ReplicatedProperties, GroundTruth);
+				if (Conflict == EAuthorityConflict::Conflict)
+				{
+					ChangeIt.RemoveCurrent();
+				}
+			}
+
+			const bool bIsRequestEmpty = Request.ObjectsToPut.IsEmpty();
+			return bIsRequestEmpty ? EBreakBehavior::Break : EBreakBehavior::Continue;
+		});
+	}
 }
diff --git a/Engine/Plugins/Developer/Concert/ConcertSync/ConcertSyncCore/Source/ConcertSyncCore/Public/Replication/AuthorityConflictSharedUtils.h b/Engine/Plugins/Developer/Concert/ConcertSync/ConcertSyncCore/Source/ConcertSyncCore/Public/Replication/AuthorityConflictSharedUtils.h
--- a/Engine/Plugins/Developer/Concert/ConcertSync/ConcertSyncCore/Source/ConcertSyncCore/Public/Replication/AuthorityConflictSharedUtils.h
+++ b/Engine/Plugins/Developer/Concert/ConcertSync/ConcertSyncCore/Source/ConcertSyncCore/Public/Replication/AuthorityConflictSharedUtils.h
@@ -7,8 +7,11 @@
 #include "Misc/EBreakBehavior.h"
 #include "Templates/Function.h"
 
+
 struct FConcertPropertyChain;
 struct FConcertPropertySelection;
+struct FConcertReplication_ChangeAuthority_Request;
+struct FConcertReplication_ChangeStream_Request;
 struct FObjectReplicationMap;
 struct FReplicatedObjectId;
 
@@ -60,4 +63,24 @@ namespace UE::ConcertSyncCore::Replication::AuthorityConflictUtils
 		const IReplicationGroundTruth& GroundTruth,
 		FProcessAuthorityConflict ProcessConflict = [](const FGuid&, const FGuid&, const FConcertPropertyChain&){ return EBreakBehavior::Break; }
 		);
+
+	/**
+	 * Removes entries from Request that would generate conflicts.
+	 * 
+	 * @param Request The request to clense
+	 * @param SendingClient The client that will send the request
+	 * @param GroundTruth Provides information about clients. On the server, it is the definite state whereas on client machines it would be
+	 * what the local client thinks what the server state is.
+	 */
+	CONCERTSYNCCORE_API void CleanseConflictsFrom(FConcertReplication_ChangeAuthority_Request& Request, const FGuid& SendingClient, const IReplicationGroundTruth& GroundTruth);
+	
+	/**
+	 * Removes entries from Request that would generate conflicts.
+	 * 
+	 * @param Request The request to clense
+	 * @param SendingClient The client that will send the request
+	 * @param GroundTruth Provides information about clients. On the server, it is the definite state whereas on client machines it would be
+	 * what the local client thinks what the server state is.
+	 */
+	CONCERTSYNCCORE_API void CleanseConflictsFrom(FConcertReplication_ChangeStream_Request& Request, const FGuid& SendingClient, const IReplicationGroundTruth& GroundTruth);
 }
