Annotate code with TODO improvements

diff --git a/Engine/Plugins/Developer/Concert/ConcertSync/ConcertSyncCore/Source/ConcertSyncCore/Public/Replication/Messages/ChangeStream.h b/Engine/Plugins/Developer/Concert/ConcertSync/ConcertSyncCore/Source/ConcertSyncCore/Public/Replication/Messages/ChangeStream.h
--- a/Engine/Plugins/Developer/Concert/ConcertSync/ConcertSyncCore/Source/ConcertSyncCore/Public/Replication/Messages/ChangeStream.h
+++ b/Engine/Plugins/Developer/Concert/ConcertSync/ConcertSyncCore/Source/ConcertSyncCore/Public/Replication/Messages/ChangeStream.h
@@ -20,6 +20,8 @@ struct FConcertReplication_ChangeStream_PutObject
 {
 	GENERATED_BODY()
 
+	// TODO UE-201166: Add two enums here that describe how Propertiy and ClassPath are to be interpreted.
+	
 	/**
 	 * The property selection the object should have.
 	 * Objects must always have a non-empty selection: use FConcertChangeStream_Request::ObjectsToRemove to remove objects.
@@ -95,6 +97,8 @@ struct FConcertReplication_ChangeStream_Request
 	 */
 	UPROPERTY()
 	TMap<FObjectInStreamID, FConcertReplication_ChangeStream_PutObject> ObjectsToPut;
+
+	// TODO UE-201167: Change StreamsToAdd so you just give FGuid of new streams and specify properties in ObjectsToPut > easier API usage.
 	
 	/**
 	 * New streams to add to the server.
