Backout enabling landscape support in Concert (original CL25978997)

diff --git a/Engine/Plugins/Developer/Concert/ConcertSync/ConcertSyncCore/Config/BaseConcertSyncCore.ini b/Engine/Plugins/Developer/Concert/ConcertSync/ConcertSyncCore/Config/BaseConcertSyncCore.ini
--- a/Engine/Plugins/Developer/Concert/ConcertSync/ConcertSyncCore/Config/BaseConcertSyncCore.ini
+++ b/Engine/Plugins/Developer/Concert/ConcertSync/ConcertSyncCore/Config/BaseConcertSyncCore.ini
@@ -10,6 +10,6 @@
 +AllowedTransientProperties="/Script/Engine.Level:bIsVisible"
 +AllowedTransientProperties="/Script/Engine.SceneComponent:AttachChildren"
 +ExcludedProperties="/Script/Engine.InstanceCacheDataBase:UniqueTransientPackage"
-
++ExcludeTransactionClassFilters=(ObjectOuterClass=None,ObjectClasses=(Class'/Script/Landscape.LandscapeProxy',Class'/Script/Landscape.LandscapeHeightfieldCollisionComponent',Class'/Script/Landscape.LandscapeComponent'))
 ; If you wish to use Composure in Multi-user, you may add that line under [/Script/ConcertSyncCore.ConcertSyncConfig] in your project engine config settings.
 ;+IncludeObjectClassFilters=(ObjectOuterClass=None,ObjectClasses=(Class'/Script/Composure.CompositingElementPass',Class'/Script/Composure.ComposurePostProcessPassPolicy'))
