Prepend reflected types with FConcertXXX and improve names

diff --git a/Engine/Plugins/Developer/Concert/ConcertApp/MultiUserClient/Source/MultiUserClient/Private/Replication/Authority/AuthoritySynchronizer_RemoteClient.cpp b/Engine/Plugins/Developer/Concert/ConcertApp/MultiUserClient/Source/MultiUserClient/Private/Replication/Authority/AuthoritySynchronizer_RemoteClient.cpp
--- a/Engine/Plugins/Developer/Concert/ConcertApp/MultiUserClient/Source/MultiUserClient/Private/Replication/Authority/AuthoritySynchronizer_RemoteClient.cpp
+++ b/Engine/Plugins/Developer/Concert/ConcertApp/MultiUserClient/Source/MultiUserClient/Private/Replication/Authority/AuthoritySynchronizer_RemoteClient.cpp
@@ -26,10 +26,10 @@ namespace UE::MultiUserClient
 		return !LastServerState.IsEmpty();
 	}
 
-	void FAuthoritySynchronizer_RemoteClient::HandleAuthorityQuery(const TArray<FReplicationAuthorityInfo>& PerStreamAuthority)
+	void FAuthoritySynchronizer_RemoteClient::HandleAuthorityQuery(const TArray<FConcertAuthorityClientInfo>& PerStreamAuthority)
 	{
 		TSet<FSoftObjectPath> OldServerState = MoveTemp(LastServerState);
-		for (const FReplicationAuthorityInfo& Info : PerStreamAuthority)
+		for (const FConcertAuthorityClientInfo& Info : PerStreamAuthority)
 		{
 			LastServerState.Append(Info.AuthoredObjects);
 		}
diff --git a/Engine/Plugins/Developer/Concert/ConcertApp/MultiUserClient/Source/MultiUserClient/Private/Replication/Authority/AuthoritySynchronizer_RemoteClient.h b/Engine/Plugins/Developer/Concert/ConcertApp/MultiUserClient/Source/MultiUserClient/Private/Replication/Authority/AuthoritySynchronizer_RemoteClient.h
--- a/Engine/Plugins/Developer/Concert/ConcertApp/MultiUserClient/Source/MultiUserClient/Private/Replication/Authority/AuthoritySynchronizer_RemoteClient.h
+++ b/Engine/Plugins/Developer/Concert/ConcertApp/MultiUserClient/Source/MultiUserClient/Private/Replication/Authority/AuthoritySynchronizer_RemoteClient.h
@@ -8,7 +8,7 @@
 #include "UObject/SoftObjectPath.h"
 
 struct FGuid;
-struct FReplicationAuthorityInfo;
+struct FConcertAuthorityClientInfo;
 
 namespace UE::MultiUserClient
 {
@@ -37,7 +37,7 @@ namespace UE::MultiUserClient
 		/** The most up to date server state of the remote client's authority. */
 		TSet<FSoftObjectPath> LastServerState; 
 
-		void HandleAuthorityQuery(const TArray<FReplicationAuthorityInfo>& PerStreamAuthority);
+		void HandleAuthorityQuery(const TArray<FConcertAuthorityClientInfo>& PerStreamAuthority);
 	};
 }
 
diff --git a/Engine/Plugins/Developer/Concert/ConcertApp/MultiUserClient/Source/MultiUserClient/Private/Replication/Frequency/FrequencyChangeTracker.cpp b/Engine/Plugins/Developer/Concert/ConcertApp/MultiUserClient/Source/MultiUserClient/Private/Replication/Frequency/FrequencyChangeTracker.cpp
--- a/Engine/Plugins/Developer/Concert/ConcertApp/MultiUserClient/Source/MultiUserClient/Private/Replication/Frequency/FrequencyChangeTracker.cpp
+++ b/Engine/Plugins/Developer/Concert/ConcertApp/MultiUserClient/Source/MultiUserClient/Private/Replication/Frequency/FrequencyChangeTracker.cpp
@@ -32,7 +32,7 @@ namespace UE::MultiUserClient
 
 		// For now we only support adding overrides.
 		const FGuid StreamId = StreamSynchronizer.GetStreamId();
-		const FObjectReplicationMap& ReplicationMap = StreamSynchronizer.GetServerState();
+		const FConcertObjectReplicationMap& ReplicationMap = StreamSynchronizer.GetServerState();
 		for (auto IteratorForAdds = Result.OverridesToAdd.CreateIterator(); IteratorForAdds; ++IteratorForAdds)
 		{
 			const TPair<FSoftObjectPath, FConcertObjectReplicationSettings>& Pair = *IteratorForAdds;
@@ -54,7 +54,7 @@ namespace UE::MultiUserClient
 
 	void FFrequencyChangeTracker::RefreshChanges()
 	{
-		const FObjectReplicationMap& ReplicationMap = StreamSynchronizer.GetServerState();
+		const FConcertObjectReplicationMap& ReplicationMap = StreamSynchronizer.GetServerState();
 		for (auto IteratorForAdds = RecordedChanges.OverridesToAdd.CreateIterator(); IteratorForAdds; ++IteratorForAdds)
 		{
 			const TPair<FSoftObjectPath, FConcertObjectReplicationSettings>& Pair = *IteratorForAdds;
diff --git a/Engine/Plugins/Developer/Concert/ConcertApp/MultiUserClient/Source/MultiUserClient/Private/Replication/MultiUserReplicationManager.cpp b/Engine/Plugins/Developer/Concert/ConcertApp/MultiUserClient/Source/MultiUserClient/Private/Replication/MultiUserReplicationManager.cpp
--- a/Engine/Plugins/Developer/Concert/ConcertApp/MultiUserClient/Source/MultiUserClient/Private/Replication/MultiUserReplicationManager.cpp
+++ b/Engine/Plugins/Developer/Concert/ConcertApp/MultiUserClient/Source/MultiUserClient/Private/Replication/MultiUserReplicationManager.cpp
@@ -100,7 +100,7 @@ namespace UE::MultiUserClient
 		OnReplicationConnectionStateChangedDelegate.Broadcast(ConnectionState);
 	}
 
-	const FObjectReplicationMap* FMultiUserReplicationManager::FindReplicationMapForClient(const FGuid& ClientId) const
+	const FConcertObjectReplicationMap* FMultiUserReplicationManager::FindReplicationMapForClient(const FGuid& ClientId) const
 	{
 		if (ConnectedState && ensureMsgf(IsInGameThread(), TEXT("To simplify implementation, only calls from game thread are allowed.")))
 		{
diff --git a/Engine/Plugins/Developer/Concert/ConcertApp/MultiUserClient/Source/MultiUserClient/Private/Replication/MultiUserReplicationManager.h b/Engine/Plugins/Developer/Concert/ConcertApp/MultiUserClient/Source/MultiUserClient/Private/Replication/MultiUserReplicationManager.h
--- a/Engine/Plugins/Developer/Concert/ConcertApp/MultiUserClient/Source/MultiUserClient/Private/Replication/MultiUserReplicationManager.h
+++ b/Engine/Plugins/Developer/Concert/ConcertApp/MultiUserClient/Source/MultiUserClient/Private/Replication/MultiUserReplicationManager.h
@@ -69,7 +69,7 @@ namespace UE::MultiUserClient
 		EMultiUserReplicationConnectionState GetConnectionState() const { return ConnectionState; }
 
 		//~ Begin IMultiUserReplication Interface
-		virtual const FObjectReplicationMap* FindReplicationMapForClient(const FGuid& ClientId) const override;
+		virtual const FConcertObjectReplicationMap* FindReplicationMapForClient(const FGuid& ClientId) const override;
 		virtual bool IsReplicatingObject(const FGuid& ClientId, const FSoftObjectPath& ObjectPath) const override;
 		virtual void RegisterReplicationDiscoverer(TSharedRef<IReplicationDiscoverer> Discoverer) override;
 		virtual void RemoveReplicationDiscoverer(const TSharedRef<IReplicationDiscoverer>& Discoverer) override;
diff --git a/Engine/Plugins/Developer/Concert/ConcertApp/MultiUserClient/Source/MultiUserClient/Private/Replication/Stream/IClientStreamSynchronizer.h b/Engine/Plugins/Developer/Concert/ConcertApp/MultiUserClient/Source/MultiUserClient/Private/Replication/Stream/IClientStreamSynchronizer.h
--- a/Engine/Plugins/Developer/Concert/ConcertApp/MultiUserClient/Source/MultiUserClient/Private/Replication/Stream/IClientStreamSynchronizer.h
+++ b/Engine/Plugins/Developer/Concert/ConcertApp/MultiUserClient/Source/MultiUserClient/Private/Replication/Stream/IClientStreamSynchronizer.h
@@ -6,7 +6,7 @@
 #include "Delegates/Delegate.h"
 
 struct FConcertStreamFrequencySettings;
-struct FObjectReplicationMap;
+struct FConcertObjectReplicationMap;
 
 namespace UE::MultiUserClient
 {
@@ -24,7 +24,7 @@ namespace UE::MultiUserClient
 		virtual FGuid GetStreamId() const = 0;
 		
 		/** @return What local instance thinks the client's server state is. */
-		virtual const FObjectReplicationMap& GetServerState() const = 0;
+		virtual const FConcertObjectReplicationMap& GetServerState() const = 0;
 
 		/** @return What the local instance thinks the client's replication frequencies are. */
 		virtual const FConcertStreamFrequencySettings& GetFrequencySettings() const = 0;
diff --git a/Engine/Plugins/Developer/Concert/ConcertApp/MultiUserClient/Source/MultiUserClient/Private/Replication/Stream/StreamChangeTracker.cpp b/Engine/Plugins/Developer/Concert/ConcertApp/MultiUserClient/Source/MultiUserClient/Private/Replication/Stream/StreamChangeTracker.cpp
--- a/Engine/Plugins/Developer/Concert/ConcertApp/MultiUserClient/Source/MultiUserClient/Private/Replication/Stream/StreamChangeTracker.cpp
+++ b/Engine/Plugins/Developer/Concert/ConcertApp/MultiUserClient/Source/MultiUserClient/Private/Replication/Stream/StreamChangeTracker.cpp
@@ -13,7 +13,7 @@ namespace UE::MultiUserClient
 {
 	FStreamChangeTracker::FStreamChangeTracker(
 		IClientStreamSynchronizer& InStreamSynchronizer,
-		TAttribute<FObjectReplicationMap*> InStreamWithInProgressChangesAttribute,
+		TAttribute<FConcertObjectReplicationMap*> InStreamWithInProgressChangesAttribute,
 		FOnModifyReplicationMap InOnModifyReplicationMapDelegate
 		)
 		: StreamSynchronizer(InStreamSynchronizer)
@@ -59,21 +59,21 @@ namespace UE::MultiUserClient
 	const FConcertPropertySelection* FStreamChangeTracker::GetPropertiesAfterSubmit(const FSoftObjectPath& ObjectPath) const
 	{
 		const FGuid StreamId = StreamSynchronizer.GetStreamId();
-		const FObjectInStreamID ObjectId { StreamId, ObjectPath };
+		const FConcertObjectInStreamID ObjectId { StreamId, ObjectPath };
 		const FConcertReplication_ChangeStream_PutObject* PutObject = CachedDeltaChange.ObjectsToPut.Find(ObjectId);
 		if (PutObject && !PutObject->Properties.ReplicatedProperties.IsEmpty())
 		{
 			return &PutObject->Properties;
 		}
 
-		const FReplicatedObjectInfo* ObjectInfo = StreamSynchronizer.GetServerState().ReplicatedObjects.Find(ObjectPath);
+		const FConcertReplicatedObjectInfo* ObjectInfo = StreamSynchronizer.GetServerState().ReplicatedObjects.Find(ObjectPath);
 		return ObjectInfo ? &ObjectInfo->PropertySelection : nullptr;
 	}
 
 	FStreamChangeTracker::EObjectChangeType FStreamChangeTracker::GetObjectChanges(const FSoftObjectPath& Object) const
 	{
 		const FGuid StreamId = StreamSynchronizer.GetStreamId();
-		const FObjectInStreamID ObjectId { StreamId, Object };
+		const FConcertObjectInStreamID ObjectId { StreamId, Object };
 		if (CachedDeltaChange.ObjectsToRemove.Contains(ObjectId))
 		{
 			return EObjectChangeType::Removed;
@@ -91,13 +91,13 @@ namespace UE::MultiUserClient
 	
 	FStreamChangelist FStreamChangeTracker::DiffChanges(
 		const FGuid& StreamId,
-		const FObjectReplicationMap& Base,
-		const FObjectReplicationMap& Changed
+		const FConcertObjectReplicationMap& Base,
+		const FConcertObjectReplicationMap& Changed
 		)
 	{
 		// BuildRequestFromDiff does not tolerate any invalid entries (like empty properties, which the UI generates right after you add an object to the list)
-		FObjectReplicationMap Cleansed = Changed;
-		ConcertSyncCore::Replication::ChangeStreamUtils::IterateInvalidEntries(Changed, [&Cleansed](const FSoftObjectPath& InvalidObject, const FReplicatedObjectInfo&)
+		FConcertObjectReplicationMap Cleansed = Changed;
+		ConcertSyncCore::Replication::ChangeStreamUtils::IterateInvalidEntries(Changed, [&Cleansed](const FSoftObjectPath& InvalidObject, const FConcertReplicatedObjectInfo&)
 		{
 			Cleansed.ReplicatedObjects.Remove(InvalidObject);
 			return EBreakBehavior::Continue;
diff --git a/Engine/Plugins/Developer/Concert/ConcertApp/MultiUserClient/Source/MultiUserClient/Private/Replication/Stream/StreamChangeTracker.h b/Engine/Plugins/Developer/Concert/ConcertApp/MultiUserClient/Source/MultiUserClient/Private/Replication/Stream/StreamChangeTracker.h
--- a/Engine/Plugins/Developer/Concert/ConcertApp/MultiUserClient/Source/MultiUserClient/Private/Replication/Stream/StreamChangeTracker.h
+++ b/Engine/Plugins/Developer/Concert/ConcertApp/MultiUserClient/Source/MultiUserClient/Private/Replication/Stream/StreamChangeTracker.h
@@ -30,7 +30,7 @@ namespace UE::MultiUserClient
 		DECLARE_DELEGATE(FOnModifyReplicationMap);
 		FStreamChangeTracker(
 			IClientStreamSynchronizer& InStreamSynchronizer,
-			TAttribute<FObjectReplicationMap*> InStreamWithInProgressChangesAttribute,
+			TAttribute<FConcertObjectReplicationMap*> InStreamWithInProgressChangesAttribute,
 			FOnModifyReplicationMap InOnModifyReplicationMapDelegate
 			);
 		~FStreamChangeTracker();
@@ -73,7 +73,7 @@ namespace UE::MultiUserClient
 		 * Represents ConfirmedServerState with changes made to it. These changes have not been sent to the server, yet.
 		 * This is cleared when a change is sent to the server. 
 		 */
-		const TAttribute<FObjectReplicationMap*> StreamWithInProgressChangesAttribute;
+		const TAttribute<FConcertObjectReplicationMap*> StreamWithInProgressChangesAttribute;
 		
 		/**
 		 * The changes that if applied to ConfirmedServerState would result in StreamWithInProgressChangesAttribute.
@@ -86,11 +86,11 @@ namespace UE::MultiUserClient
 
 		FStreamChangelist DiffChanges() const
 		{
-			const FObjectReplicationMap* Map = StreamWithInProgressChangesAttribute.Get();
+			const FConcertObjectReplicationMap* Map = StreamWithInProgressChangesAttribute.Get();
 			return Map ? DiffChanges(StreamSynchronizer.GetStreamId(), StreamSynchronizer.GetServerState(), *Map) : FStreamChangelist{};
 		}
 		/** Builds the changelist to get from Base to Changed. */
-		static FStreamChangelist DiffChanges(const FGuid& StreamId, const FObjectReplicationMap& Base, const FObjectReplicationMap& Changed);
+		static FStreamChangelist DiffChanges(const FGuid& StreamId, const FConcertObjectReplicationMap& Base, const FConcertObjectReplicationMap& Changed);
 	};
 }
 
diff --git a/Engine/Plugins/Developer/Concert/ConcertApp/MultiUserClient/Source/MultiUserClient/Private/Replication/Stream/StreamSynchronizer_LocalClient.cpp b/Engine/Plugins/Developer/Concert/ConcertApp/MultiUserClient/Source/MultiUserClient/Private/Replication/Stream/StreamSynchronizer_LocalClient.cpp
--- a/Engine/Plugins/Developer/Concert/ConcertApp/MultiUserClient/Source/MultiUserClient/Private/Replication/Stream/StreamSynchronizer_LocalClient.cpp
+++ b/Engine/Plugins/Developer/Concert/ConcertApp/MultiUserClient/Source/MultiUserClient/Private/Replication/Stream/StreamSynchronizer_LocalClient.cpp
@@ -31,15 +31,15 @@ namespace UE::MultiUserClient
 		}
 	}
 
-	const FObjectReplicationMap& FStreamSynchronizer_LocalClient::GetServerState() const
+	const FConcertObjectReplicationMap& FStreamSynchronizer_LocalClient::GetServerState() const
 	{
-		const FSharedReplicationStreamDescription* Description = GetLocalMultiUserStream();
+		const FConcertBaseStreamInfo* Description = GetLocalMultiUserStream();
 		return Description ? Description->ReplicationMap : EmptyState.ReplicationMap;
 	}
 
 	const FConcertStreamFrequencySettings& FStreamSynchronizer_LocalClient::GetFrequencySettings() const
 	{
-		const FSharedReplicationStreamDescription* Description = GetLocalMultiUserStream();
+		const FConcertBaseStreamInfo* Description = GetLocalMultiUserStream();
 		return Description ? Description->FrequencySettings : EmptyState.FrequencySettings;
 	}
 
@@ -49,7 +49,7 @@ namespace UE::MultiUserClient
 		OnServerStateChangedDelegate.Broadcast();
 	}
 
-	const FSharedReplicationStreamDescription* FStreamSynchronizer_LocalClient::GetLocalMultiUserStream() const
+	const FConcertBaseStreamInfo* FStreamSynchronizer_LocalClient::GetLocalMultiUserStream() const
 	{
 		IConcertClientReplicationManager* ReplicationManager = LocalClient->GetReplicationManager();
 		if (!ensure(ReplicationManager))
@@ -57,8 +57,8 @@ namespace UE::MultiUserClient
 			return nullptr;
 		}
 		
-		const FSharedReplicationStreamDescription* Result = nullptr;
-		ReplicationManager->ForEachRegisteredStream([this, &Result](const FReplicationStreamDescription& Stream)
+		const FConcertBaseStreamInfo* Result = nullptr;
+		ReplicationManager->ForEachRegisteredStream([this, &Result](const FConcertReplicationStream& Stream)
 		{
 			if (Stream.BaseDescription.Identifier == LocalClientStreamId)
 			{
diff --git a/Engine/Plugins/Developer/Concert/ConcertApp/MultiUserClient/Source/MultiUserClient/Private/Replication/Stream/StreamSynchronizer_LocalClient.h b/Engine/Plugins/Developer/Concert/ConcertApp/MultiUserClient/Source/MultiUserClient/Private/Replication/Stream/StreamSynchronizer_LocalClient.h
--- a/Engine/Plugins/Developer/Concert/ConcertApp/MultiUserClient/Source/MultiUserClient/Private/Replication/Stream/StreamSynchronizer_LocalClient.h
+++ b/Engine/Plugins/Developer/Concert/ConcertApp/MultiUserClient/Source/MultiUserClient/Private/Replication/Stream/StreamSynchronizer_LocalClient.h
@@ -4,7 +4,7 @@
 
 #include "IClientStreamSynchronizer.h"
 #include "Replication/Data/ObjectReplicationMap.h"
-#include "Replication/Data/ReplicationStreamDescription.h"
+#include "Replication/Data/ReplicationStream.h"
 #include "Templates/UnrealTemplate.h"
 
 class IConcertSyncClient;
@@ -21,7 +21,7 @@ namespace UE::MultiUserClient
 
 		//~ Begin IClientStreamSynchronizer Interface
 		virtual FGuid GetStreamId() const override { return LocalClientStreamId; }
-		virtual const FObjectReplicationMap& GetServerState() const override;
+		virtual const FConcertObjectReplicationMap& GetServerState() const override;
 		virtual const FConcertStreamFrequencySettings& GetFrequencySettings() const override;
 		virtual FOnServerStateChanged& OnServerStateChanged() override { return OnServerStateChangedDelegate; }
 		//~ End IClientStreamSynchronizer Interface
@@ -29,7 +29,7 @@ namespace UE::MultiUserClient
 	private:
 		
 		/** Returned by GetServerState when there is no registered stream. */
-		const FSharedReplicationStreamDescription EmptyState; 
+		const FConcertBaseStreamInfo EmptyState; 
 		
 		/** Owning client. Used to send change requests to the server. */
 		const TSharedRef<IConcertSyncClient> LocalClient;
@@ -40,7 +40,7 @@ namespace UE::MultiUserClient
 		FOnServerStateChanged OnServerStateChangedDelegate;
 		
 		void OnPostStreamsChanged();
-		const FSharedReplicationStreamDescription* GetLocalMultiUserStream() const;
+		const FConcertBaseStreamInfo* GetLocalMultiUserStream() const;
 	};
 }
 
diff --git a/Engine/Plugins/Developer/Concert/ConcertApp/MultiUserClient/Source/MultiUserClient/Private/Replication/Stream/StreamSynchronizer_RemoteClient.cpp b/Engine/Plugins/Developer/Concert/ConcertApp/MultiUserClient/Source/MultiUserClient/Private/Replication/Stream/StreamSynchronizer_RemoteClient.cpp
--- a/Engine/Plugins/Developer/Concert/ConcertApp/MultiUserClient/Source/MultiUserClient/Private/Replication/Stream/StreamSynchronizer_RemoteClient.cpp
+++ b/Engine/Plugins/Developer/Concert/ConcertApp/MultiUserClient/Source/MultiUserClient/Private/Replication/Stream/StreamSynchronizer_RemoteClient.cpp
@@ -30,11 +30,11 @@ namespace UE::MultiUserClient
 		return UMultiUserReplicationClientPreset::MultiUserStreamID;
 	}
 
-	void FStreamSynchronizer_RemoteClient::HandleStreamQuery(const TArray<FSharedReplicationStreamDescription>& Streams)
+	void FStreamSynchronizer_RemoteClient::HandleStreamQuery(const TArray<FConcertBaseStreamInfo>& Streams)
 	{
 		// MU clients use UMultiUserReplicationClientPreset::MultiUserStreamID for streams.
 		// That handles the (unlikely) case in which some external logic had added streams to the same client which we must differentiate
-		const FSharedReplicationStreamDescription* MultiUserDescription = Streams.FindByPredicate([](const FSharedReplicationStreamDescription& Description)
+		const FConcertBaseStreamInfo* MultiUserDescription = Streams.FindByPredicate([](const FConcertBaseStreamInfo& Description)
 		{
 			return Description.Identifier == UMultiUserReplicationClientPreset::MultiUserStreamID;
 		});
diff --git a/Engine/Plugins/Developer/Concert/ConcertApp/MultiUserClient/Source/MultiUserClient/Private/Replication/Stream/StreamSynchronizer_RemoteClient.h b/Engine/Plugins/Developer/Concert/ConcertApp/MultiUserClient/Source/MultiUserClient/Private/Replication/Stream/StreamSynchronizer_RemoteClient.h
--- a/Engine/Plugins/Developer/Concert/ConcertApp/MultiUserClient/Source/MultiUserClient/Private/Replication/Stream/StreamSynchronizer_RemoteClient.h
+++ b/Engine/Plugins/Developer/Concert/ConcertApp/MultiUserClient/Source/MultiUserClient/Private/Replication/Stream/StreamSynchronizer_RemoteClient.h
@@ -3,8 +3,7 @@
 #pragma once
 
 #include "IClientStreamSynchronizer.h"
-#include "Replication/Data/ObjectReplicationMap.h"
-#include "Replication/Data/ReplicationStreamDescription.h"
+#include "Replication/Data/ReplicationStream.h"
 
 #include "Containers/Array.h"
 #include "Templates/UnrealTemplate.h"
@@ -23,7 +22,7 @@ namespace UE::MultiUserClient
 
 		//~ Begin IClientStreamSynchronizer Interface
 		virtual FGuid GetStreamId() const override;
-		virtual const FObjectReplicationMap& GetServerState() const override { return LastKnownServerState.ReplicationMap; }
+		virtual const FConcertObjectReplicationMap& GetServerState() const override { return LastKnownServerState.ReplicationMap; }
 		virtual const FConcertStreamFrequencySettings& GetFrequencySettings() const override { return LastKnownServerState.FrequencySettings; }
 		virtual FOnServerStateChanged& OnServerStateChanged() override { return OnServerStateChangedDelegate; }
 		//~ End IClientStreamSynchronizer Interface
@@ -36,13 +35,13 @@ namespace UE::MultiUserClient
 		const FDelegateHandle QueryStreamHandle; 
 
 		/** Represents what the local client thinks the replication map on the server currently looks like. */
-		FSharedReplicationStreamDescription LastKnownServerState;
+		FConcertBaseStreamInfo LastKnownServerState;
 		
 		/** Event executed when the result of GetServerState has been synched. */
 		FOnServerStateChanged OnServerStateChangedDelegate;
 
 		/** Called in regular intervals with the contents of the remote client. */
-		void HandleStreamQuery(const TArray<FSharedReplicationStreamDescription>& Streams);
+		void HandleStreamQuery(const TArray<FConcertBaseStreamInfo>& Streams);
 	};
 }
 
diff --git a/Engine/Plugins/Developer/Concert/ConcertApp/MultiUserClient/Source/MultiUserClient/Private/Replication/Submission/External/ClientChangeConversionUtils.cpp b/Engine/Plugins/Developer/Concert/ConcertApp/MultiUserClient/Source/MultiUserClient/Private/Replication/Submission/External/ClientChangeConversionUtils.cpp
--- a/Engine/Plugins/Developer/Concert/ConcertApp/MultiUserClient/Source/MultiUserClient/Private/Replication/Submission/External/ClientChangeConversionUtils.cpp
+++ b/Engine/Plugins/Developer/Concert/ConcertApp/MultiUserClient/Source/MultiUserClient/Private/Replication/Submission/External/ClientChangeConversionUtils.cpp
@@ -17,14 +17,14 @@ namespace UE::MultiUserClient::ClientChangeConversionUtils
 		static TOptional<FConcertReplication_ChangeStream_PutObject> GeneratePutRequest(
 			const UObject& Object,
 			FPropertyChange&& PropertyChange,
-			const FObjectReplicationMap& ClientStreamContent
+			const FConcertObjectReplicationMap& ClientStreamContent
 			)
 		{
 			// TODO UE-201166: This step could be simplified if we had an append operation in FConcertReplication_ChangeStream_PutObject
-			const FReplicatedObjectInfo* ExistingObjectInfo = ClientStreamContent.ReplicatedObjects.Find(&Object);
+			const FConcertReplicatedObjectInfo* ExistingObjectInfo = ClientStreamContent.ReplicatedObjects.Find(&Object);
 			if (ExistingObjectInfo)
 			{
-				FReplicatedObjectInfo NewObjectInfo = FReplicatedObjectInfo::Make(Object);
+				FConcertReplicatedObjectInfo NewObjectInfo = FConcertReplicatedObjectInfo::Make(Object);
 				TArray<FConcertPropertyChain>& ReplicatedProperties = NewObjectInfo.PropertySelection.ReplicatedProperties;
 				switch (PropertyChange.ChangeType)
 				{
@@ -46,7 +46,7 @@ namespace UE::MultiUserClient::ClientChangeConversionUtils
 			}
 			else
 			{
-				FReplicatedObjectInfo NewObjectInfo = FReplicatedObjectInfo::Make(Object);
+				FConcertReplicatedObjectInfo NewObjectInfo = FConcertReplicatedObjectInfo::Make(Object);
 				TArray<FConcertPropertyChain>& ReplicatedProperties = NewObjectInfo.PropertySelection.ReplicatedProperties;
 				switch (PropertyChange.ChangeType)
 				{
@@ -66,7 +66,7 @@ namespace UE::MultiUserClient::ClientChangeConversionUtils
 		static void BuildStreamChanges(
 			TMap<UObject*, FPropertyChange>&& PropertyChanges,
 			const FGuid& ClientStreamId,
-			const FObjectReplicationMap& ClientStreamContent,
+			const FConcertObjectReplicationMap& ClientStreamContent,
 			FStreamChangelist& StreamChangelist
 			)
 		{
@@ -86,7 +86,7 @@ namespace UE::MultiUserClient::ClientChangeConversionUtils
 					PutRequest->Properties.DiscoverAndAddImplicitParentProperties();
 					
 					StreamChangelist.ObjectsToPut.Add(
-						FObjectInStreamID { ClientStreamId, Object },
+						FConcertObjectInStreamID { ClientStreamId, Object },
 						MoveTemp(*PutRequest)
 					);
 				}
@@ -98,7 +98,7 @@ namespace UE::MultiUserClient::ClientChangeConversionUtils
 			for (FSoftObjectPath& RemovedObject : ObjectsToRemove)
 			{
 				StreamChangelist.ObjectsToRemove.Emplace(
-					FObjectInStreamID{ ClientStreamId, MoveTemp(RemovedObject) }
+					FConcertObjectInStreamID{ ClientStreamId, MoveTemp(RemovedObject) }
 				);
 			}
 		}
@@ -107,7 +107,7 @@ namespace UE::MultiUserClient::ClientChangeConversionUtils
 	TOptional<FConcertReplication_ChangeStream_Request> Transform(
 		FChangeStreamRequest Request,
 		const FGuid& ClientStreamId,
-		const FObjectReplicationMap& ClientStreamContent
+		const FConcertObjectReplicationMap& ClientStreamContent
 		)
 	{
 		FStreamChangelist StreamChangelist;
diff --git a/Engine/Plugins/Developer/Concert/ConcertApp/MultiUserClient/Source/MultiUserClient/Private/Replication/Submission/External/ClientChangeConversionUtils.h b/Engine/Plugins/Developer/Concert/ConcertApp/MultiUserClient/Source/MultiUserClient/Private/Replication/Submission/External/ClientChangeConversionUtils.h
--- a/Engine/Plugins/Developer/Concert/ConcertApp/MultiUserClient/Source/MultiUserClient/Private/Replication/Submission/External/ClientChangeConversionUtils.h
+++ b/Engine/Plugins/Developer/Concert/ConcertApp/MultiUserClient/Source/MultiUserClient/Private/Replication/Submission/External/ClientChangeConversionUtils.h
@@ -8,7 +8,7 @@
 
 #include "Misc/Optional.h"
 
-struct FObjectReplicationMap;
+struct FConcertObjectReplicationMap;
 struct FGuid;
 
 namespace UE::MultiUserClient
@@ -29,7 +29,7 @@ namespace UE::MultiUserClient::ClientChangeConversionUtils
 	TOptional<FConcertReplication_ChangeStream_Request> Transform(
 		FChangeStreamRequest Request,
 		const FGuid& ClientStreamId,
-		const FObjectReplicationMap& ClientStreamContent
+		const FConcertObjectReplicationMap& ClientStreamContent
 		);
 	
 	/** Transforms FChangeAuthorityRequest to FConcertReplication_ChangeStream_Request, if the request is correctly formatted. */
@@ -48,7 +48,7 @@ namespace UE::MultiUserClient::ClientChangeConversionUtils
 	inline TOptional<FConcertReplication_ChangeStream_Request> Transform(
 		TOptional<FChangeStreamRequest> Request,
 		const FGuid& ClientStreamId,
-		const FObjectReplicationMap& ClientStreamContent
+		const FConcertObjectReplicationMap& ClientStreamContent
 		)
 	{
 		return Request
diff --git a/Engine/Plugins/Developer/Concert/ConcertApp/MultiUserClient/Source/MultiUserClient/Private/Replication/Submission/External/ExternalClientChangeDelegates.h b/Engine/Plugins/Developer/Concert/ConcertApp/MultiUserClient/Source/MultiUserClient/Private/Replication/Submission/External/ExternalClientChangeDelegates.h
--- a/Engine/Plugins/Developer/Concert/ConcertApp/MultiUserClient/Source/MultiUserClient/Private/Replication/Submission/External/ExternalClientChangeDelegates.h
+++ b/Engine/Plugins/Developer/Concert/ConcertApp/MultiUserClient/Source/MultiUserClient/Private/Replication/Submission/External/ExternalClientChangeDelegates.h
@@ -4,9 +4,9 @@
 
 #include "Delegates/Delegate.h"
 
-struct FObjectReplicationMap;
+struct FConcertObjectReplicationMap;
 
 namespace UE::MultiUserClient
 {
-	DECLARE_DELEGATE_RetVal(const FObjectReplicationMap*, FGetStreamContent);
+	DECLARE_DELEGATE_RetVal(const FConcertObjectReplicationMap*, FGetStreamContent);
 }
diff --git a/Engine/Plugins/Developer/Concert/ConcertApp/MultiUserClient/Source/MultiUserClient/Private/Replication/Submission/MultiEdit/ReassignObjectPropertiesLogic.cpp b/Engine/Plugins/Developer/Concert/ConcertApp/MultiUserClient/Source/MultiUserClient/Private/Replication/Submission/MultiEdit/ReassignObjectPropertiesLogic.cpp
--- a/Engine/Plugins/Developer/Concert/ConcertApp/MultiUserClient/Source/MultiUserClient/Private/Replication/Submission/MultiEdit/ReassignObjectPropertiesLogic.cpp
+++ b/Engine/Plugins/Developer/Concert/ConcertApp/MultiUserClient/Source/MultiUserClient/Private/Replication/Submission/MultiEdit/ReassignObjectPropertiesLogic.cpp
@@ -22,7 +22,7 @@ namespace UE::MultiUserClient
 		struct FObjectReassignment
 		{
 			/** Used to build the final changelist to apply to the target client. Maps ClientIds to the object data they had registered. */
-			TMap<FClientId, FObjectReplicationMap> OldRegisteredObjects;
+			TMap<FClientId, FConcertObjectReplicationMap> OldRegisteredObjects;
 			/** The frequency settings clients when reassignment operation was started. */
 			TMap<FClientId, FConcertStreamFrequencySettings> OldFrequencies;
 			/** Changes to make to all the other clients */
@@ -35,14 +35,14 @@ namespace UE::MultiUserClient
 		static FObjectReassignment BuildChangesForTransferal(const FReplicationClientManager& ClientManager, TConstArrayView<FSoftObjectPath> ObjectsToReassign, const FGuid& StreamId)
 		{
 			FObjectReassignment Result;
-			TMap<FGuid, FObjectReplicationMap>& OldRegisteredObjects = Result.OldRegisteredObjects;
+			TMap<FGuid, FConcertObjectReplicationMap>& OldRegisteredObjects = Result.OldRegisteredObjects;
 			TMap<FGuid, FConcertStreamFrequencySettings>& OldFrequencies = Result.OldFrequencies;
 			TMap<FGuid, FConcertReplication_ChangeStream_Request>& ReassignedClientRequests = Result.ReassignedClientRequests;
 			TMap<FGuid, TArray<FSoftObjectPath>>& ReassignedAuthority = Result.ReassignedAuthority;
 		
 			for (const FSoftObjectPath& ObjectPath : ObjectsToReassign)
 			{
-				const FObjectInStreamID ObjectId{ StreamId, ObjectPath };
+				const FConcertObjectInStreamID ObjectId{ StreamId, ObjectPath };
 
 				const FGlobalAuthorityCache& AuthorityCache = ClientManager.GetAuthorityCache();
 				AuthorityCache.ForEachClientWithObjectInStream(ObjectPath,
@@ -57,7 +57,7 @@ namespace UE::MultiUserClient
 						}
 						
 						const IClientStreamSynchronizer& ReassignedStreamSynchronizer = ClientToReassignFrom->GetStreamSynchronizer();
-						const FReplicatedObjectInfo* ObjectInfo = ReassignedStreamSynchronizer.GetServerState().ReplicatedObjects.Find(ObjectId.Object);
+						const FConcertReplicatedObjectInfo* ObjectInfo = ReassignedStreamSynchronizer.GetServerState().ReplicatedObjects.Find(ObjectId.Object);
 						if (!ensureMsgf(ObjectInfo, TEXT("ForEachClientWithObjectInStream lied")))
 						{
 							return EBreakBehavior::Continue;
@@ -115,7 +115,7 @@ namespace UE::MultiUserClient
 
 		/** Builds a changelist based on which remote clients we managed to change successfully. */
 		static TPair<FStreamChangelist, FFrequencyChangelist> MakeChangelistFromAppliedChanges(
-			const TMap<FGuid, FObjectReplicationMap>& OldRegisteredObjects,
+			const TMap<FGuid, FConcertObjectReplicationMap>& OldRegisteredObjects,
 			const TMap<FClientId, FConcertStreamFrequencySettings>& OldFrequencies,
 			const FParallelExecutionResult& ParallelExecutionResult,
 			const IClientStreamSynchronizer& AssignedToStreamSynchronizer
@@ -124,7 +124,7 @@ namespace UE::MultiUserClient
 			FStreamChangelist ObjectChanges;
 			FFrequencyChangelist FrequencyChanges;
 			
-			const FObjectReplicationMap& CurrentTargetClientState = AssignedToStreamSynchronizer.GetServerState();
+			const FConcertObjectReplicationMap& CurrentTargetClientState = AssignedToStreamSynchronizer.GetServerState();
 			const FConcertStreamFrequencySettings& CurrentFrequencySettings = AssignedToStreamSynchronizer.GetFrequencySettings();
 			const FGuid TargetStreamId = AssignedToStreamSynchronizer.GetStreamId();
 			for (const TPair<FGuid, FSubmitStreamChangesResponse>& ChangesRequestedOnRemote : ParallelExecutionResult.StreamResponses)
@@ -135,11 +135,11 @@ namespace UE::MultiUserClient
 					continue;
 				}
 				
-				const FObjectReplicationMap& ObjectReplicationMap = OldRegisteredObjects[ChangesRequestedOnRemote.Key];
-				for (const TPair<FSoftObjectPath, FReplicatedObjectInfo>& ChangesToApply : ObjectReplicationMap.ReplicatedObjects)
+				const FConcertObjectReplicationMap& ObjectReplicationMap = OldRegisteredObjects[ChangesRequestedOnRemote.Key];
+				for (const TPair<FSoftObjectPath, FConcertReplicatedObjectInfo>& ChangesToApply : ObjectReplicationMap.ReplicatedObjects)
 				{
 					const FSoftObjectPath& ObjectPath = ChangesToApply.Key;
-					const FObjectInStreamID ObjectId { TargetStreamId, ObjectPath };
+					const FConcertObjectInStreamID ObjectId { TargetStreamId, ObjectPath };
 					TOptional<FConcertReplication_ChangeStream_PutObject> PutRequest = FConcertReplication_ChangeStream_PutObject::MakeFromInfo(ChangesToApply.Value);
 					if (!ensure(PutRequest))
 					{
@@ -148,7 +148,7 @@ namespace UE::MultiUserClient
 
 					// We want to append the other client's properties to the ones the target client already has.
 					// TODO UE-201166: This step would not be necessary if we had an append operation in FConcertReplication_ChangeStream_PutObject
-					const FReplicatedObjectInfo* CurrentInfo = CurrentTargetClientState.ReplicatedObjects.Find(ObjectId.Object);
+					const FConcertReplicatedObjectInfo* CurrentInfo = CurrentTargetClientState.ReplicatedObjects.Find(ObjectId.Object);
 					if (CurrentInfo)
 					{
 						for (const FConcertPropertyChain& PropertyChain : CurrentInfo->PropertySelection.ReplicatedProperties)
@@ -203,16 +203,16 @@ namespace UE::MultiUserClient
 
 		/** @return Whether TargetClientObjects includes all properties ClientsToConsider have registered for ObjectToCheck. */
 		static bool DoesTargetIncludeOthers_SingleObject(
-			const FObjectReplicationMap& TargetClientObjects,
+			const FConcertObjectReplicationMap& TargetClientObjects,
 			const TArray<TNonNullPtr<const FReplicationClient>> ClientsToConsider,
 			const FSoftObjectPath& ObjectToCheck)
 		{
-			const FReplicatedObjectInfo* ObjectInfo = TargetClientObjects.ReplicatedObjects.Find(ObjectToCheck);
+			const FConcertReplicatedObjectInfo* ObjectInfo = TargetClientObjects.ReplicatedObjects.Find(ObjectToCheck);
 				
 			for (const FReplicationClient* OtherClient : ClientsToConsider)
 			{
-				const FObjectReplicationMap& OtherClientObjects = OtherClient->GetStreamSynchronizer().GetServerState();
-				const FReplicatedObjectInfo* OtherObjectInfo = OtherClientObjects.ReplicatedObjects.Find(ObjectToCheck);
+				const FConcertObjectReplicationMap& OtherClientObjects = OtherClient->GetStreamSynchronizer().GetServerState();
+				const FConcertReplicatedObjectInfo* OtherObjectInfo = OtherClientObjects.ReplicatedObjects.Find(ObjectToCheck);
 					
 				// TargetClient has at least as much if OtherClient has nothing
 				if (!OtherObjectInfo || OtherObjectInfo->PropertySelection.ReplicatedProperties.IsEmpty())
@@ -243,7 +243,7 @@ namespace UE::MultiUserClient
 				return false;
 			}
 
-			const FObjectReplicationMap& TargetClientObjects = TargetClient->GetStreamSynchronizer().GetServerState();
+			const FConcertObjectReplicationMap& TargetClientObjects = TargetClient->GetStreamSynchronizer().GetServerState();
 			const TArray<TNonNullPtr<const FReplicationClient>> ClientsToConsider = ClientManager.GetClients(ShouldConsiderClientPredicate);
 			return Algo::AllOf(ObjectsToCheck, [&TargetClientObjects, &ClientsToConsider](const FSoftObjectPath& ObjectToCheck)
 			{
@@ -260,10 +260,10 @@ namespace UE::MultiUserClient
 				return false;
 			}
 
-			const FObjectReplicationMap& TargetClientObjects = TargetClient->GetStreamSynchronizer().GetServerState();
+			const FConcertObjectReplicationMap& TargetClientObjects = TargetClient->GetStreamSynchronizer().GetServerState();
 			return Algo::AnyOf(ObjectsToCheck, [&TargetClientObjects](const FSoftObjectPath& ObjectPath)
 			{
-				const FReplicatedObjectInfo* ObjectInfo = TargetClientObjects.ReplicatedObjects.Find(ObjectPath);
+				const FConcertReplicatedObjectInfo* ObjectInfo = TargetClientObjects.ReplicatedObjects.Find(ObjectPath);
 				return ObjectInfo && !ObjectInfo->PropertySelection.ReplicatedProperties.IsEmpty();
 			});
 		}
@@ -419,9 +419,9 @@ namespace UE::MultiUserClient
 			return false;
 		}
 
-		return Algo::AnyOf(InProgressOperation->OldRegisteredObjects, [&ObjectPath](const TPair<FGuid, FObjectReplicationMap>& Change)
+		return Algo::AnyOf(InProgressOperation->OldRegisteredObjects, [&ObjectPath](const TPair<FGuid, FConcertObjectReplicationMap>& Change)
 		{
-			const FReplicatedObjectInfo* ObjectInfo = Change.Value.ReplicatedObjects.Find(ObjectPath);
+			const FConcertReplicatedObjectInfo* ObjectInfo = Change.Value.ReplicatedObjects.Find(ObjectPath);
 			return ObjectInfo && !ObjectInfo->PropertySelection.ReplicatedProperties.IsEmpty();
 		});
 	}
diff --git a/Engine/Plugins/Developer/Concert/ConcertApp/MultiUserClient/Source/MultiUserClient/Private/Replication/Submission/MultiEdit/ReassignObjectPropertiesLogic.h b/Engine/Plugins/Developer/Concert/ConcertApp/MultiUserClient/Source/MultiUserClient/Private/Replication/Submission/MultiEdit/ReassignObjectPropertiesLogic.h
--- a/Engine/Plugins/Developer/Concert/ConcertApp/MultiUserClient/Source/MultiUserClient/Private/Replication/Submission/MultiEdit/ReassignObjectPropertiesLogic.h
+++ b/Engine/Plugins/Developer/Concert/ConcertApp/MultiUserClient/Source/MultiUserClient/Private/Replication/Submission/MultiEdit/ReassignObjectPropertiesLogic.h
@@ -80,7 +80,7 @@ namespace UE::MultiUserClient
 			const TArray<FSoftObjectPath> ObjectsBeingReassigned;
 			const FGuid AssignedToClient;
 			/** The objects we tried to remove from all the other clients. If the removal request worked, the (previously) assigned properties will be transferred to the target client. */
-			const TMap<FClientId, FObjectReplicationMap> OldRegisteredObjects;
+			const TMap<FClientId, FConcertObjectReplicationMap> OldRegisteredObjects;
 			/** The frequency settings that clients had at the beginning of the reassign operation. */
 			const TMap<FClientId, FConcertStreamFrequencySettings> OldFrequencies; 
 			/** Clients to the objects they had authority over. Used to later to determine which objects to give the target client authority over (do not transfer authority from failed clients). */
diff --git a/Engine/Plugins/Developer/Concert/ConcertApp/MultiUserClient/Source/MultiUserClient/Private/Replication/Submission/Notification/SubmissionNotifier.cpp b/Engine/Plugins/Developer/Concert/ConcertApp/MultiUserClient/Source/MultiUserClient/Private/Replication/Submission/Notification/SubmissionNotifier.cpp
--- a/Engine/Plugins/Developer/Concert/ConcertApp/MultiUserClient/Source/MultiUserClient/Private/Replication/Submission/Notification/SubmissionNotifier.cpp
+++ b/Engine/Plugins/Developer/Concert/ConcertApp/MultiUserClient/Source/MultiUserClient/Private/Replication/Submission/Notification/SubmissionNotifier.cpp
@@ -107,11 +107,11 @@ namespace UE::MultiUserClient
 		switch(CompletedOp.ErrorCode)
 		{
 		case EStreamSubmissionErrorCode::Success:
-			for (const TPair<FObjectInStreamID, FReplicatedObjectId>& Pair : CompletedOp.SubmissionInfo->Response.AuthorityConflicts)
+			for (const TPair<FConcertObjectInStreamID, FConcertReplicatedObjectId>& Pair : CompletedOp.SubmissionInfo->Response.AuthorityConflicts)
 			{
 				StreamErrors.AuthorityConflicts.FindOrAdd(Pair.Key.Object, 0) += 1;
 			}
-			for (const TPair<FObjectInStreamID, EConcertPutObjectErrorCode>& Pair : CompletedOp.SubmissionInfo->Response.ObjectsToPutSemanticErrors)
+			for (const TPair<FConcertObjectInStreamID, EConcertPutObjectErrorCode>& Pair : CompletedOp.SubmissionInfo->Response.ObjectsToPutSemanticErrors)
 			{
 				StreamErrors.SemanticErrors.FindOrAdd(Pair.Key.Object, 0) += 1;
 			}
diff --git a/Engine/Plugins/Developer/Concert/ConcertApp/MultiUserClient/Source/MultiUserClient/Private/Replication/Util/GlobalAuthorityCache.cpp b/Engine/Plugins/Developer/Concert/ConcertApp/MultiUserClient/Source/MultiUserClient/Private/Replication/Util/GlobalAuthorityCache.cpp
--- a/Engine/Plugins/Developer/Concert/ConcertApp/MultiUserClient/Source/MultiUserClient/Private/Replication/Util/GlobalAuthorityCache.cpp
+++ b/Engine/Plugins/Developer/Concert/ConcertApp/MultiUserClient/Source/MultiUserClient/Private/Replication/Util/GlobalAuthorityCache.cpp
@@ -135,7 +135,7 @@ namespace UE::MultiUserClient
 				return EBreakBehavior::Continue;
 			}
 			
-			const FReplicatedObjectInfo* ObjectInfo = Client->GetStreamSynchronizer().GetServerState().ReplicatedObjects.Find(Object);
+			const FConcertReplicatedObjectInfo* ObjectInfo = Client->GetStreamSynchronizer().GetServerState().ReplicatedObjects.Find(Object);
 			const bool bHasObjectRegistered = ensureMsgf(ObjectInfo, TEXT("OnStreamChanged should have updated OwnedObjectsToClients"))
 				&& ObjectInfo->PropertySelection.ReplicatedProperties.Contains(Property);
 			const bool bHasAuthority = Client->GetAuthoritySynchronizer().HasAuthorityOver(Object);
@@ -180,9 +180,9 @@ namespace UE::MultiUserClient
 			return;
 		}
 		
-		const FObjectReplicationMap& ClientObjectMap = Client->GetStreamSynchronizer().GetServerState();
+		const FConcertObjectReplicationMap& ClientObjectMap = Client->GetStreamSynchronizer().GetServerState();
 		const IClientAuthoritySynchronizer& ClientAuthority = Client->GetAuthoritySynchronizer();
-		for (const TPair<FSoftObjectPath, FReplicatedObjectInfo>& StreamContents : ClientObjectMap.ReplicatedObjects)
+		for (const TPair<FSoftObjectPath, FConcertReplicatedObjectInfo>& StreamContents : ClientObjectMap.ReplicatedObjects)
 		{
 			const FSoftObjectPath& Object = StreamContents.Key;
 			RegisteredObjectsToClients.FindOrAdd(Object).Add(ClientId);
@@ -214,7 +214,7 @@ namespace UE::MultiUserClient
 		}
 	}
 	
-	void FGlobalAuthorityCache::ForEachStream(const FGuid& ClientEndpointId, TFunctionRef<EBreakBehavior(const FGuid& StreamId, const FObjectReplicationMap& ReplicationMap)> Callback) const
+	void FGlobalAuthorityCache::ForEachStream(const FGuid& ClientEndpointId, TFunctionRef<EBreakBehavior(const FGuid& StreamId, const FConcertObjectReplicationMap& ReplicationMap)> Callback) const
 	{
 		const FReplicationClient* Client = ClientManager.FindClient(ClientEndpointId);
 		if (!ensure(Client))
diff --git a/Engine/Plugins/Developer/Concert/ConcertApp/MultiUserClient/Source/MultiUserClient/Private/Replication/Util/GlobalAuthorityCache.h b/Engine/Plugins/Developer/Concert/ConcertApp/MultiUserClient/Source/MultiUserClient/Private/Replication/Util/GlobalAuthorityCache.h
--- a/Engine/Plugins/Developer/Concert/ConcertApp/MultiUserClient/Source/MultiUserClient/Private/Replication/Util/GlobalAuthorityCache.h
+++ b/Engine/Plugins/Developer/Concert/ConcertApp/MultiUserClient/Source/MultiUserClient/Private/Replication/Util/GlobalAuthorityCache.h
@@ -117,7 +117,7 @@ namespace UE::MultiUserClient
 		}
 
 		//~ Begin IReplicationGroundTruth Interface
-		virtual void ForEachStream(const FGuid& ClientEndpointId, TFunctionRef<EBreakBehavior(const FGuid& StreamId, const FObjectReplicationMap& ReplicationMap)> Callback) const override;
+		virtual void ForEachStream(const FGuid& ClientEndpointId, TFunctionRef<EBreakBehavior(const FGuid& StreamId, const FConcertObjectReplicationMap& ReplicationMap)> Callback) const override;
 		virtual void ForEachSendingClient(TFunctionRef<EBreakBehavior(const FGuid& ClientEndpointId)> Callback) const override;
 		virtual bool HasAuthority(const FGuid& ClientId, const FGuid& StreamId, const FSoftObjectPath& ObjectPath) const override;
 		//~ End IReplicationGroundTruth Interface
diff --git a/Engine/Plugins/Developer/Concert/ConcertApp/MultiUserClient/Source/MultiUserClient/Private/Replication/Util/RegularQueryService.cpp b/Engine/Plugins/Developer/Concert/ConcertApp/MultiUserClient/Source/MultiUserClient/Private/Replication/Util/RegularQueryService.cpp
--- a/Engine/Plugins/Developer/Concert/ConcertApp/MultiUserClient/Source/MultiUserClient/Private/Replication/Util/RegularQueryService.cpp
+++ b/Engine/Plugins/Developer/Concert/ConcertApp/MultiUserClient/Source/MultiUserClient/Private/Replication/Util/RegularQueryService.cpp
@@ -32,7 +32,7 @@ namespace UE::MultiUserClient
 	FDelegateHandle FRegularQueryService::RegisterStreamQuery(const FGuid& EndpointId, FStreamQueryDelegate Delegate)
 	{
 		FStreamQueryInfo& Info = StreamQueryInfos.FindOrAdd(EndpointId);
-		const FDelegateHandle DelegateHandle = Info.Delegate.AddLambda([Delegate = MoveTemp(Delegate)](const TArray<FSharedReplicationStreamDescription>& Descriptions)
+		const FDelegateHandle DelegateHandle = Info.Delegate.AddLambda([Delegate = MoveTemp(Delegate)](const TArray<FConcertBaseStreamInfo>& Descriptions)
 		{
 			Delegate.Execute(Descriptions);
 		});
@@ -43,7 +43,7 @@ namespace UE::MultiUserClient
 	FDelegateHandle FRegularQueryService::RegisterAuthorityQuery(const FGuid& EndpointId, FAuthorityQueryDelegate Delegate)
 	{
 		FAuthorityQueryInfo& Info = AuthorityQueryInfos.FindOrAdd(EndpointId);
-		const FDelegateHandle DelegateHandle = Info.Delegate.AddLambda([Delegate = MoveTemp(Delegate)](const TArray<FReplicationAuthorityInfo>& Infos)
+		const FDelegateHandle DelegateHandle = Info.Delegate.AddLambda([Delegate = MoveTemp(Delegate)](const TArray<FConcertAuthorityClientInfo>& Infos)
 		{
 			Delegate.Execute(Infos);
 		});
@@ -160,7 +160,7 @@ namespace UE::MultiUserClient
 	{
 		TGuardValue<bool> Guard(bIsHandlingQueryResponse, true);
 		
-		for (const TPair<FGuid, FReplicationClientQueriedInfo>& StreamQueryPair : Response.ClientInfo)
+		for (const TPair<FGuid, FConcertQueriedClientInfo>& StreamQueryPair : Response.ClientInfo)
 		{
 			const FGuid StreamId = StreamQueryPair.Key;
 			
diff --git a/Engine/Plugins/Developer/Concert/ConcertApp/MultiUserClient/Source/MultiUserClient/Private/Replication/Util/RegularQueryService.h b/Engine/Plugins/Developer/Concert/ConcertApp/MultiUserClient/Source/MultiUserClient/Private/Replication/Util/RegularQueryService.h
--- a/Engine/Plugins/Developer/Concert/ConcertApp/MultiUserClient/Source/MultiUserClient/Private/Replication/Util/RegularQueryService.h
+++ b/Engine/Plugins/Developer/Concert/ConcertApp/MultiUserClient/Source/MultiUserClient/Private/Replication/Util/RegularQueryService.h
@@ -12,8 +12,8 @@ class IConcertSyncClient;
 
 namespace UE::MultiUserClient
 {
-	DECLARE_DELEGATE_OneParam(FStreamQueryDelegate, const TArray<FSharedReplicationStreamDescription>&);
-	DECLARE_DELEGATE_OneParam(FAuthorityQueryDelegate, const TArray<FReplicationAuthorityInfo>&);
+	DECLARE_DELEGATE_OneParam(FStreamQueryDelegate, const TArray<FConcertBaseStreamInfo>&);
+	DECLARE_DELEGATE_OneParam(FAuthorityQueryDelegate, const TArray<FConcertAuthorityClientInfo>&);
 	
 	/** Sends regular FConcertReplication_QueryReplicationInfo_Request to endpoints and publishes the results. */
 	class FRegularQueryService : public FNoncopyable
@@ -51,8 +51,8 @@ namespace UE::MultiUserClient
 		const FTSTicker::FDelegateHandle TickerDelegateHandle;
 
 		// We use a multicast delegate here because it handles unsubscribing while being Broadcast
-		DECLARE_MULTICAST_DELEGATE_OneParam(FMulticastStreamQueryDelegate, const TArray<FSharedReplicationStreamDescription>&);
-		DECLARE_MULTICAST_DELEGATE_OneParam(FMulticastAuthorityQueryDelegate, const TArray<FReplicationAuthorityInfo>&);
+		DECLARE_MULTICAST_DELEGATE_OneParam(FMulticastStreamQueryDelegate, const TArray<FConcertBaseStreamInfo>&);
+		DECLARE_MULTICAST_DELEGATE_OneParam(FMulticastAuthorityQueryDelegate, const TArray<FConcertAuthorityClientInfo>&);
 		struct FStreamQueryInfo
 		{
 			TSet<FDelegateHandle> Handles;
diff --git a/Engine/Plugins/Developer/Concert/ConcertApp/MultiUserClient/Source/MultiUserClient/Private/Replication/Util/StreamRequestUtils.cpp b/Engine/Plugins/Developer/Concert/ConcertApp/MultiUserClient/Source/MultiUserClient/Private/Replication/Util/StreamRequestUtils.cpp
--- a/Engine/Plugins/Developer/Concert/ConcertApp/MultiUserClient/Source/MultiUserClient/Private/Replication/Util/StreamRequestUtils.cpp
+++ b/Engine/Plugins/Developer/Concert/ConcertApp/MultiUserClient/Source/MultiUserClient/Private/Replication/Util/StreamRequestUtils.cpp
@@ -13,19 +13,19 @@ namespace UE::MultiUserClient::StreamRequestUtils
 		FFrequencyChangelist FrequencyChanges
 		)
 	{
-		const TMap<FObjectInStreamID, FConcertReplication_ChangeStream_PutObject>& ObjectsToPut = ObjectChanges.ObjectsToPut;
+		const TMap<FConcertObjectInStreamID, FConcertReplication_ChangeStream_PutObject>& ObjectsToPut = ObjectChanges.ObjectsToPut;
 		
 		FConcertReplication_ChangeStream_Request Request;
 		Request.StreamsToAdd.Emplace();
-		FReplicationStreamDescription& NewStream = Request.StreamsToAdd[0];
+		FConcertReplicationStream& NewStream = Request.StreamsToAdd[0];
 		NewStream.BaseDescription.Identifier = StreamId;
 			
 		// If creating a new stream, the objects must be supplied in the description instead of in PutObjects!
-		FObjectReplicationMap& ReplicationMap = NewStream.BaseDescription.ReplicationMap;
+		FConcertObjectReplicationMap& ReplicationMap = NewStream.BaseDescription.ReplicationMap;
 		ReplicationMap.ReplicatedObjects.Reserve(ObjectsToPut.Num());
-		for (const TPair<FObjectInStreamID, FConcertReplication_ChangeStream_PutObject>& PutObjectPair : ObjectsToPut)
+		for (const TPair<FConcertObjectInStreamID, FConcertReplication_ChangeStream_PutObject>& PutObjectPair : ObjectsToPut)
 		{
-			const TOptional<FReplicatedObjectInfo> NewObjectInfo = PutObjectPair.Value.MakeObjectInfoIfValid();
+			const TOptional<FConcertReplicatedObjectInfo> NewObjectInfo = PutObjectPair.Value.MakeObjectInfoIfValid();
 			// The editing UI allows adding objects without properties (to make UX easier) - do not submit those to the server.
 			if (!NewObjectInfo)
 			{
diff --git a/Engine/Plugins/Developer/Concert/ConcertApp/MultiUserClient/Source/MultiUserClient/Private/Replication/Util/StreamRequestUtils.h b/Engine/Plugins/Developer/Concert/ConcertApp/MultiUserClient/Source/MultiUserClient/Private/Replication/Util/StreamRequestUtils.h
--- a/Engine/Plugins/Developer/Concert/ConcertApp/MultiUserClient/Source/MultiUserClient/Private/Replication/Util/StreamRequestUtils.h
+++ b/Engine/Plugins/Developer/Concert/ConcertApp/MultiUserClient/Source/MultiUserClient/Private/Replication/Util/StreamRequestUtils.h
@@ -11,8 +11,8 @@ namespace UE::MultiUserClient
 	/** Describes changes that MU clients make to the objects in a stream. */
 	struct FStreamChangelist
 	{
-		TSet<FObjectInStreamID> ObjectsToRemove;
-		TMap<FObjectInStreamID, FConcertReplication_ChangeStream_PutObject> ObjectsToPut;
+		TSet<FConcertObjectInStreamID> ObjectsToRemove;
+		TMap<FConcertObjectInStreamID, FConcertReplication_ChangeStream_PutObject> ObjectsToPut;
 	};
 	
 	/** Describes changes that MU clients make to the frequency settings in a stream. */
diff --git a/Engine/Plugins/Developer/Concert/ConcertApp/MultiUserClient/Source/MultiUserClient/Private/Widgets/ActiveSession/Replication/Client/Multi/Columns/SAssignPropertyComboBox.cpp b/Engine/Plugins/Developer/Concert/ConcertApp/MultiUserClient/Source/MultiUserClient/Private/Widgets/ActiveSession/Replication/Client/Multi/Columns/SAssignPropertyComboBox.cpp
--- a/Engine/Plugins/Developer/Concert/ConcertApp/MultiUserClient/Source/MultiUserClient/Private/Widgets/ActiveSession/Replication/Client/Multi/Columns/SAssignPropertyComboBox.cpp
+++ b/Engine/Plugins/Developer/Concert/ConcertApp/MultiUserClient/Source/MultiUserClient/Private/Widgets/ActiveSession/Replication/Client/Multi/Columns/SAssignPropertyComboBox.cpp
@@ -29,10 +29,10 @@ namespace UE::MultiUserClient
 			TArray<FGuid> Clients;
 			ClientManager.ForEachClient([&DisplayedProperty, &EditedObjects, &Clients](const FReplicationClient& Client)
 			{
-				const TMap<FSoftObjectPath, FReplicatedObjectInfo>& ObjectInfoMap = Client.GetStreamSynchronizer().GetServerState().ReplicatedObjects;
+				const TMap<FSoftObjectPath, FConcertReplicatedObjectInfo>& ObjectInfoMap = Client.GetStreamSynchronizer().GetServerState().ReplicatedObjects;
 				for (const FSoftObjectPath& ObjectPath : EditedObjects)
 				{
-					if (const FReplicatedObjectInfo* ObjectInfo = ObjectInfoMap.Find(ObjectPath)
+					if (const FConcertReplicatedObjectInfo* ObjectInfo = ObjectInfoMap.Find(ObjectPath)
 						; ObjectInfo && ObjectInfo->PropertySelection.ReplicatedProperties.Contains(DisplayedProperty))
 					{
 						Clients.Add(Client.GetEndpointId());
diff --git a/Engine/Plugins/Developer/Concert/ConcertApp/MultiUserClient/Source/MultiUserClient/Public/Replication/IMultiUserReplication.h b/Engine/Plugins/Developer/Concert/ConcertApp/MultiUserClient/Source/MultiUserClient/Public/Replication/IMultiUserReplication.h
--- a/Engine/Plugins/Developer/Concert/ConcertApp/MultiUserClient/Source/MultiUserClient/Public/Replication/IMultiUserReplication.h
+++ b/Engine/Plugins/Developer/Concert/ConcertApp/MultiUserClient/Source/MultiUserClient/Public/Replication/IMultiUserReplication.h
@@ -6,7 +6,7 @@
 #include "Templates/SharedPointer.h"
 
 struct FGuid;
-struct FObjectReplicationMap;
+struct FConcertObjectReplicationMap;
 
 namespace UE::MultiUserClient
 {
@@ -26,7 +26,7 @@ namespace UE::MultiUserClient
 		 * This server state is regularly polled whilst the local client state should always be in synch.
 		 * This function must be called from the game thread.
 		 */
-		virtual const FObjectReplicationMap* FindReplicationMapForClient(const FGuid& ClientId) const = 0;
+		virtual const FConcertObjectReplicationMap* FindReplicationMapForClient(const FGuid& ClientId) const = 0;
 		/**
 		 * @return Whether the local editor instance thinks the client has authority over the properties it has registered to ObjectPath.
 		 * This function must be called from the game thread.
diff --git a/Engine/Plugins/Developer/Concert/ConcertApp/MultiUserClient/Source/MultiUserClientLibrary/Private/Replication/MultiUserReplicationSubsystem.cpp b/Engine/Plugins/Developer/Concert/ConcertApp/MultiUserClient/Source/MultiUserClientLibrary/Private/Replication/MultiUserReplicationSubsystem.cpp
--- a/Engine/Plugins/Developer/Concert/ConcertApp/MultiUserClient/Source/MultiUserClientLibrary/Private/Replication/MultiUserReplicationSubsystem.cpp
+++ b/Engine/Plugins/Developer/Concert/ConcertApp/MultiUserClient/Source/MultiUserClientLibrary/Private/Replication/MultiUserReplicationSubsystem.cpp
@@ -29,13 +29,13 @@ TArray<FConcertPropertyChainWrapper> UMultiUserReplicationSubsystem::GetProperti
 	const UE::MultiUserClient::IMultiUserReplication* ReplicationInterface = IMultiUserClientModule::Get().GetReplication();
 	if (ensureMsgf(ReplicationInterface, TEXT("We expected it to always be valid.")))
 	{
-		const FObjectReplicationMap* ObjectMap = ReplicationInterface->FindReplicationMapForClient(ClientId);
+		const FConcertObjectReplicationMap* ObjectMap = ReplicationInterface->FindReplicationMapForClient(ClientId);
 		if (!ObjectMap)
 		{
 			return {};
 		}
 
-		const FReplicatedObjectInfo* ObjectInfo = ObjectMap->ReplicatedObjects.Find(ObjectPath);
+		const FConcertReplicatedObjectInfo* ObjectInfo = ObjectMap->ReplicatedObjects.Find(ObjectPath);
 		if (!ObjectInfo)
 		{
 			return {};
@@ -58,7 +58,7 @@ TArray<FSoftObjectPath> UMultiUserReplicationSubsystem::GetRegisteredObjects(con
 	const UE::MultiUserClient::IMultiUserReplication* ReplicationInterface = IMultiUserClientModule::Get().GetReplication();
 	if (ensureMsgf(ReplicationInterface, TEXT("We expected it to always be valid.")))
 	{
-		const FObjectReplicationMap* ObjectMap = ReplicationInterface->FindReplicationMapForClient(ClientId);
+		const FConcertObjectReplicationMap* ObjectMap = ReplicationInterface->FindReplicationMapForClient(ClientId);
 		if (!ObjectMap)
 		{
 			return {};
@@ -78,14 +78,14 @@ TArray<FSoftObjectPath> UMultiUserReplicationSubsystem::GetReplicatedObjects(con
 	const UE::MultiUserClient::IMultiUserReplication* ReplicationInterface = IMultiUserClientModule::Get().GetReplication();
 	if (ensureMsgf(ReplicationInterface, TEXT("We expected it to always be valid.")))
 	{
-		const FObjectReplicationMap* ObjectMap = ReplicationInterface->FindReplicationMapForClient(ClientId);
+		const FConcertObjectReplicationMap* ObjectMap = ReplicationInterface->FindReplicationMapForClient(ClientId);
 		if (!ObjectMap)
 		{
 			return {};
 		}
 
 		TArray<FSoftObjectPath> Result;
-		for (const TPair<FSoftObjectPath, FReplicatedObjectInfo>& Pair : ObjectMap->ReplicatedObjects)
+		for (const TPair<FSoftObjectPath, FConcertReplicatedObjectInfo>& Pair : ObjectMap->ReplicatedObjects)
 		{
 			if (ReplicationInterface->IsReplicatingObject(ClientId, Pair.Key))
 			{
diff --git a/Engine/Plugins/Developer/Concert/ConcertApp/MultiUserClient/Source/MultiUserReplication/Public/MultiUserPropertyReplicationSelection.h b/Engine/Plugins/Developer/Concert/ConcertApp/MultiUserClient/Source/MultiUserReplication/Public/MultiUserPropertyReplicationSelection.h
--- a/Engine/Plugins/Developer/Concert/ConcertApp/MultiUserClient/Source/MultiUserReplication/Public/MultiUserPropertyReplicationSelection.h
+++ b/Engine/Plugins/Developer/Concert/ConcertApp/MultiUserClient/Source/MultiUserReplication/Public/MultiUserPropertyReplicationSelection.h
@@ -19,5 +19,5 @@ class MULTIUSERREPLICATION_API UMultiUserPropertyReplicationSelection : public U
 public:
 
 	UPROPERTY()
-	FObjectReplicationMap ReplicationMap;
+	FConcertObjectReplicationMap ReplicationMap;
 };
diff --git a/Engine/Plugins/Developer/Concert/ConcertApp/MultiUserClient/Source/MultiUserReplicationEditor/Internal/Assets/MultiUserReplicationClientPreset.h b/Engine/Plugins/Developer/Concert/ConcertApp/MultiUserClient/Source/MultiUserReplicationEditor/Internal/Assets/MultiUserReplicationClientPreset.h
--- a/Engine/Plugins/Developer/Concert/ConcertApp/MultiUserClient/Source/MultiUserReplicationEditor/Internal/Assets/MultiUserReplicationClientPreset.h
+++ b/Engine/Plugins/Developer/Concert/ConcertApp/MultiUserClient/Source/MultiUserReplicationEditor/Internal/Assets/MultiUserReplicationClientPreset.h
@@ -4,7 +4,7 @@
 
 #include "UObject/Object.h"
 #include "MultiUserReplicationStream.h"
-#include "Replication/Data/ReplicationStreamDescription.h"
+#include "Replication/Data/ReplicationStream.h"
 #include "MultiUserReplicationClientPreset.generated.h"
 
 /**
@@ -29,5 +29,5 @@ public:
 	void ClearClient();
 
 	/** Generates a description that can be sent to the MU server. */
-	FReplicationStreamDescription GenerateDescription() const;
+	FConcertReplicationStream GenerateDescription() const;
 };
diff --git a/Engine/Plugins/Developer/Concert/ConcertApp/MultiUserClient/Source/MultiUserReplicationEditor/Internal/Assets/MultiUserReplicationStream.h b/Engine/Plugins/Developer/Concert/ConcertApp/MultiUserClient/Source/MultiUserReplicationEditor/Internal/Assets/MultiUserReplicationStream.h
--- a/Engine/Plugins/Developer/Concert/ConcertApp/MultiUserClient/Source/MultiUserReplicationEditor/Internal/Assets/MultiUserReplicationStream.h
+++ b/Engine/Plugins/Developer/Concert/ConcertApp/MultiUserClient/Source/MultiUserReplicationEditor/Internal/Assets/MultiUserReplicationStream.h
@@ -3,11 +3,11 @@
 #pragma once
 
 #include "UObject/Object.h"
-#include "Replication/Data/ReplicationStreamDescription.h"
+#include "Replication/Data/ReplicationStream.h"
 #include "Misc/Attribute.h"
 #include "MultiUserReplicationStream.generated.h"
 
-/** Wraps FObjectReplicationMap so its edition can be transacted in the editor. */
+/** Wraps FConcertObjectReplicationMap so its edition can be transacted in the editor. */
 UCLASS()
 class MULTIUSERREPLICATIONEDITOR_API UMultiUserReplicationStream : public UObject
 {
@@ -20,15 +20,15 @@ public:
 
 	/** The objects this stream will modify. */
 	UPROPERTY()
-	FObjectReplicationMap ReplicationMap;
+	FConcertObjectReplicationMap ReplicationMap;
 
 	UMultiUserReplicationStream();
 	
-	FReplicationStreamDescription GenerateDescription() const;
+	FConcertReplicationStream GenerateDescription() const;
 	
-	TAttribute<FObjectReplicationMap*> MakeReplicationMapGetterAttribute()
+	TAttribute<FConcertObjectReplicationMap*> MakeReplicationMapGetterAttribute()
 	{
-		return TAttribute<FObjectReplicationMap*>::CreateLambda([WeakThis = TWeakObjectPtr<UMultiUserReplicationStream>(this)]()
+		return TAttribute<FConcertObjectReplicationMap*>::CreateLambda([WeakThis = TWeakObjectPtr<UMultiUserReplicationStream>(this)]()
 		{
 			return WeakThis.IsValid() ? &WeakThis->ReplicationMap : nullptr;
 		});
diff --git a/Engine/Plugins/Developer/Concert/ConcertApp/MultiUserClient/Source/MultiUserReplicationEditor/Private/Assets/MultiUserReplicationClientPreset.cpp b/Engine/Plugins/Developer/Concert/ConcertApp/MultiUserClient/Source/MultiUserReplicationEditor/Private/Assets/MultiUserReplicationClientPreset.cpp
--- a/Engine/Plugins/Developer/Concert/ConcertApp/MultiUserClient/Source/MultiUserReplicationEditor/Private/Assets/MultiUserReplicationClientPreset.cpp
+++ b/Engine/Plugins/Developer/Concert/ConcertApp/MultiUserClient/Source/MultiUserReplicationEditor/Private/Assets/MultiUserReplicationClientPreset.cpp
@@ -3,7 +3,7 @@
 #include "Assets/MultiUserReplicationClientPreset.h"
 
 #include "Assets/MultiUserReplicationStream.h"
-#include "Replication/Data/ReplicationStreamDescription.h"
+#include "Replication/Data/ReplicationStream.h"
 
 UMultiUserReplicationClientPreset::UMultiUserReplicationClientPreset()
 {
@@ -17,7 +17,7 @@ void UMultiUserReplicationClientPreset::ClearClient()
 	Stream->ReplicationMap.ReplicatedObjects.Empty();
 }
 
-FReplicationStreamDescription UMultiUserReplicationClientPreset::GenerateDescription() const
+FConcertReplicationStream UMultiUserReplicationClientPreset::GenerateDescription() const
 {
 	return Stream->GenerateDescription();
 }
diff --git a/Engine/Plugins/Developer/Concert/ConcertApp/MultiUserClient/Source/MultiUserReplicationEditor/Private/Assets/MultiUserReplicationStream.cpp b/Engine/Plugins/Developer/Concert/ConcertApp/MultiUserClient/Source/MultiUserReplicationEditor/Private/Assets/MultiUserReplicationStream.cpp
--- a/Engine/Plugins/Developer/Concert/ConcertApp/MultiUserClient/Source/MultiUserReplicationEditor/Private/Assets/MultiUserReplicationStream.cpp
+++ b/Engine/Plugins/Developer/Concert/ConcertApp/MultiUserClient/Source/MultiUserReplicationEditor/Private/Assets/MultiUserReplicationStream.cpp
@@ -2,7 +2,7 @@
 
 #include "Assets/MultiUserReplicationStream.h"
 
-#include "Replication/Data/ReplicationStreamDescription.h"
+#include "Replication/Data/ReplicationStream.h"
 
 UMultiUserReplicationStream::UMultiUserReplicationStream()
 {
@@ -12,7 +12,7 @@ UMultiUserReplicationStream::UMultiUserReplicationStream()
 		: FGuid::NewGuid();
 }
 
-FReplicationStreamDescription UMultiUserReplicationStream::GenerateDescription() const
+FConcertReplicationStream UMultiUserReplicationStream::GenerateDescription() const
 {
 	return { StreamId, ReplicationMap };
 }
diff --git a/Engine/Plugins/Developer/Concert/ConcertApp/MultiUserClient/Source/MultiUserReplicationEditor/Private/StreamEditor/ReplicationSessionPresetEditorToolkit.h b/Engine/Plugins/Developer/Concert/ConcertApp/MultiUserClient/Source/MultiUserReplicationEditor/Private/StreamEditor/ReplicationSessionPresetEditorToolkit.h
--- a/Engine/Plugins/Developer/Concert/ConcertApp/MultiUserClient/Source/MultiUserReplicationEditor/Private/StreamEditor/ReplicationSessionPresetEditorToolkit.h
+++ b/Engine/Plugins/Developer/Concert/ConcertApp/MultiUserClient/Source/MultiUserReplicationEditor/Private/StreamEditor/ReplicationSessionPresetEditorToolkit.h
@@ -5,7 +5,7 @@
 #include "Tools/BaseAssetToolkit.h"
 
 class UMultiUserReplicationSessionPreset;
-struct FObjectReplicationMap;
+struct FConcertObjectReplicationMap;
 
 namespace UE::MultiUserReplicationEditor
 {
diff --git a/Engine/Plugins/Developer/Concert/ConcertSync/ConcertSyncClient/Source/ConcertSyncClient/Private/Replication/IConcertClientReplicationManager.cpp b/Engine/Plugins/Developer/Concert/ConcertSync/ConcertSyncClient/Source/ConcertSyncClient/Private/Replication/IConcertClientReplicationManager.cpp
--- a/Engine/Plugins/Developer/Concert/ConcertSync/ConcertSyncClient/Source/ConcertSyncClient/Private/Replication/IConcertClientReplicationManager.cpp
+++ b/Engine/Plugins/Developer/Concert/ConcertSync/ConcertSyncClient/Source/ConcertSyncClient/Private/Replication/IConcertClientReplicationManager.cpp
@@ -13,7 +13,7 @@ namespace UE::ConcertSyncClient::Replication::Private
 	{
 		for (const FSoftObjectPath& ObjectPath : Objects)
 		{
-			Manager.ForEachRegisteredStream([&InCallback, &ObjectPath](const FReplicationStreamDescription& StreamDescription)
+			Manager.ForEachRegisteredStream([&InCallback, &ObjectPath](const FConcertReplicationStream& StreamDescription)
 			{
 				const bool bStreamContainsObject = StreamDescription.BaseDescription.ReplicationMap.ReplicatedObjects.Contains(ObjectPath);
 				if (bStreamContainsObject)
@@ -35,10 +35,10 @@ bool IConcertClientReplicationManager::HasRegisteredStreams() const
 	return ForEachRegisteredStream([](const auto&){ return EBreakBehavior::Break; }) == EStreamEnumerationResult::Iterated;
 }
 
-TArray<FReplicationStreamDescription> IConcertClientReplicationManager::GetRegisteredStreams() const
+TArray<FConcertReplicationStream> IConcertClientReplicationManager::GetRegisteredStreams() const
 {
-	TArray<FReplicationStreamDescription> Result;
-	ForEachRegisteredStream([&Result](const FReplicationStreamDescription& Description){ Result.Add(Description); return EBreakBehavior::Continue; });
+	TArray<FConcertReplicationStream> Result;
+	ForEachRegisteredStream([&Result](const FConcertReplicationStream& Description){ Result.Add(Description); return EBreakBehavior::Continue; });
 	return Result;
 }
 
diff --git a/Engine/Plugins/Developer/Concert/ConcertSync/ConcertSyncClient/Source/ConcertSyncClient/Private/Replication/Manager/ReplicationManager.cpp b/Engine/Plugins/Developer/Concert/ConcertSync/ConcertSyncClient/Source/ConcertSyncClient/Private/Replication/Manager/ReplicationManager.cpp
--- a/Engine/Plugins/Developer/Concert/ConcertSync/ConcertSyncClient/Source/ConcertSyncClient/Private/Replication/Manager/ReplicationManager.cpp
+++ b/Engine/Plugins/Developer/Concert/ConcertSync/ConcertSyncClient/Source/ConcertSyncClient/Private/Replication/Manager/ReplicationManager.cpp
@@ -52,7 +52,7 @@ namespace UE::ConcertSyncClient::Replication
 	}
 
 	IConcertClientReplicationManager::EStreamEnumerationResult FReplicationManager::ForEachRegisteredStream(
-		TFunctionRef<EBreakBehavior(const FReplicationStreamDescription& Stream)> Callback
+		TFunctionRef<EBreakBehavior(const FConcertReplicationStream& Stream)> Callback
 		) const
 	{
 		return ensureMsgf(CurrentState, TEXT("StartAcceptingJoinRequests should have been called at this point."))
diff --git a/Engine/Plugins/Developer/Concert/ConcertSync/ConcertSyncClient/Source/ConcertSyncClient/Private/Replication/Manager/ReplicationManager.h b/Engine/Plugins/Developer/Concert/ConcertSync/ConcertSyncClient/Source/ConcertSyncClient/Private/Replication/Manager/ReplicationManager.h
--- a/Engine/Plugins/Developer/Concert/ConcertSync/ConcertSyncClient/Source/ConcertSyncClient/Private/Replication/Manager/ReplicationManager.h
+++ b/Engine/Plugins/Developer/Concert/ConcertSync/ConcertSyncClient/Source/ConcertSyncClient/Private/Replication/Manager/ReplicationManager.h
@@ -29,7 +29,7 @@ namespace UE::ConcertSyncClient::Replication
 		virtual void LeaveReplicationSession() override;
 		virtual bool CanJoin() override;
 		virtual bool IsConnectedToReplicationSession() override;
-		virtual EStreamEnumerationResult ForEachRegisteredStream(TFunctionRef<EBreakBehavior(const FReplicationStreamDescription& Stream)> Callback) const override;
+		virtual EStreamEnumerationResult ForEachRegisteredStream(TFunctionRef<EBreakBehavior(const FConcertReplicationStream& Stream)> Callback) const override;
 		virtual TFuture<FConcertReplication_ChangeAuthority_Response> RequestAuthorityChange(FConcertReplication_ChangeAuthority_Request Args) override;
 		virtual TFuture<FConcertReplication_QueryReplicationInfo_Response> QueryClientInfo(FConcertReplication_QueryReplicationInfo_Request Args) override;
 		virtual TFuture<FConcertReplication_ChangeStream_Response> ChangeStream(FConcertReplication_ChangeStream_Request Args) override;
diff --git a/Engine/Plugins/Developer/Concert/ConcertSync/ConcertSyncClient/Source/ConcertSyncClient/Private/Replication/Manager/ReplicationManagerState.h b/Engine/Plugins/Developer/Concert/ConcertSync/ConcertSyncClient/Source/ConcertSyncClient/Private/Replication/Manager/ReplicationManagerState.h
--- a/Engine/Plugins/Developer/Concert/ConcertSync/ConcertSyncClient/Source/ConcertSyncClient/Private/Replication/Manager/ReplicationManagerState.h
+++ b/Engine/Plugins/Developer/Concert/ConcertSync/ConcertSyncClient/Source/ConcertSyncClient/Private/Replication/Manager/ReplicationManagerState.h
@@ -22,7 +22,7 @@ namespace UE::ConcertSyncClient::Replication
 		// Default implementations for subclasses in which the operation is not valid
 		
 		//~ Begin IConcertClientReplicationManager Interface
-		virtual EStreamEnumerationResult ForEachRegisteredStream(TFunctionRef<EBreakBehavior(const FReplicationStreamDescription& Stream)> Callback) const override { return EStreamEnumerationResult::NoRegisteredStreams; }
+		virtual EStreamEnumerationResult ForEachRegisteredStream(TFunctionRef<EBreakBehavior(const FConcertReplicationStream& Stream)> Callback) const override { return EStreamEnumerationResult::NoRegisteredStreams; }
 		virtual TFuture<FConcertReplication_ChangeAuthority_Response> RequestAuthorityChange(FConcertReplication_ChangeAuthority_Request Args) override { return RejectAll(MoveTemp(Args)); }
 		virtual TFuture<FConcertReplication_QueryReplicationInfo_Response> QueryClientInfo(FConcertReplication_QueryReplicationInfo_Request Args) override { return MakeFulfilledPromise<FConcertReplication_QueryReplicationInfo_Response>().GetFuture(); }
 		virtual TFuture<FConcertReplication_ChangeStream_Response> ChangeStream(FConcertReplication_ChangeStream_Request Args) override { return MakeFulfilledPromise<FConcertReplication_ChangeStream_Response>().GetFuture(); }
diff --git a/Engine/Plugins/Developer/Concert/ConcertSync/ConcertSyncClient/Source/ConcertSyncClient/Private/Replication/Manager/ReplicationManagerState_Connected.cpp b/Engine/Plugins/Developer/Concert/ConcertSync/ConcertSyncClient/Source/ConcertSyncClient/Private/Replication/Manager/ReplicationManagerState_Connected.cpp
--- a/Engine/Plugins/Developer/Concert/ConcertSync/ConcertSyncClient/Source/ConcertSyncClient/Private/Replication/Manager/ReplicationManagerState_Connected.cpp
+++ b/Engine/Plugins/Developer/Concert/ConcertSync/ConcertSyncClient/Source/ConcertSyncClient/Private/Replication/Manager/ReplicationManagerState_Connected.cpp
@@ -46,7 +46,7 @@ namespace UE::ConcertSyncClient::Replication
 	FReplicationManagerState_Connected::FReplicationManagerState_Connected(
 		TSharedRef<IConcertClientSession> LiveSession,
 		IConcertClientReplicationBridge* ReplicationBridge,
-		TArray<FReplicationStreamDescription> StreamDescriptions,
+		TArray<FConcertReplicationStream> StreamDescriptions,
 		FReplicationManager& Owner
 		)
 		: FReplicationManagerState(Owner)
@@ -93,10 +93,10 @@ namespace UE::ConcertSyncClient::Replication
 	}
 
 	IConcertClientReplicationManager::EStreamEnumerationResult FReplicationManagerState_Connected::ForEachRegisteredStream(
-		TFunctionRef<EBreakBehavior(const FReplicationStreamDescription& Stream)> Callback
+		TFunctionRef<EBreakBehavior(const FConcertReplicationStream& Stream)> Callback
 		) const
 	{
-		for (const FReplicationStreamDescription& Stream : RegisteredStreams)
+		for (const FConcertReplicationStream& Stream : RegisteredStreams)
 		{
 			if (Callback(Stream) == EBreakBehavior::Break)
 			{
@@ -235,13 +235,13 @@ namespace UE::ConcertSyncClient::Replication
 		
 		// Build RegisteredStreams while RegisteredStreams has the old, unupdated state
 		TMap<FSoftObjectPath, TArray<FGuid>> BundledModifiedObjects;
-		for (const TPair<FObjectInStreamID, FConcertReplication_ChangeStream_PutObject>& PutObject : Request.ObjectsToPut)
+		for (const TPair<FConcertObjectInStreamID, FConcertReplication_ChangeStream_PutObject>& PutObject : Request.ObjectsToPut)
 		{
-			const FObjectInStreamID ObjectInfo = PutObject.Key;
+			const FConcertObjectInStreamID ObjectInfo = PutObject.Key;
 			const FSoftObjectPath Object = ObjectInfo.Object;
 			const FGuid StreamId = ObjectInfo.StreamId;
 			
-			const FReplicationStreamDescription* StreamDescription = RegisteredStreams.FindByPredicate([&StreamId](const FReplicationStreamDescription& Stream)
+			const FConcertReplicationStream* StreamDescription = RegisteredStreams.FindByPredicate([&StreamId](const FConcertReplicationStream& Stream)
 			{
 				return Stream.BaseDescription.Identifier == StreamId;
 			});
@@ -268,7 +268,7 @@ namespace UE::ConcertSyncClient::Replication
 		static void ForEachObjectRemovedFromStreams(const FConcertReplication_ChangeStream_Request& Request, TFunctionRef<void(const FSoftObjectPath& ObjectPath, const TArray<FGuid>& Streams)> Callback)
 		{
 			TMap<FSoftObjectPath, TArray<FGuid>> BundledRemovedObjects;
-			for (const FObjectInStreamID& RemovedObject : Request.ObjectsToRemove)
+			for (const FConcertObjectInStreamID& RemovedObject : Request.ObjectsToRemove)
 			{
 				BundledRemovedObjects.FindOrAdd(RemovedObject.Object).Add(RemovedObject.StreamId);
 			}
@@ -341,9 +341,9 @@ namespace UE::ConcertSyncClient::Replication
 		}
 	}
 
-	FConcertObjectReplicationSettings FReplicationManagerState_Connected::GetObjectFrequencySettings(const FReplicatedObjectId& Object) const
+	FConcertObjectReplicationSettings FReplicationManagerState_Connected::GetObjectFrequencySettings(const FConcertReplicatedObjectId& Object) const
 	{
-		const FReplicationStreamDescription* Stream = RegisteredStreams.FindByPredicate([&Object](const FReplicationStreamDescription& Description)
+		const FConcertReplicationStream* Stream = RegisteredStreams.FindByPredicate([&Object](const FConcertReplicationStream& Description)
 		{
 			return Description.BaseDescription.Identifier == Object.StreamId;
 		});
diff --git a/Engine/Plugins/Developer/Concert/ConcertSync/ConcertSyncClient/Source/ConcertSyncClient/Private/Replication/Manager/ReplicationManagerState_Connected.h b/Engine/Plugins/Developer/Concert/ConcertSync/ConcertSyncClient/Source/ConcertSyncClient/Private/Replication/Manager/ReplicationManagerState_Connected.h
--- a/Engine/Plugins/Developer/Concert/ConcertSync/ConcertSyncClient/Source/ConcertSyncClient/Private/Replication/Manager/ReplicationManagerState_Connected.h
+++ b/Engine/Plugins/Developer/Concert/ConcertSync/ConcertSyncClient/Source/ConcertSyncClient/Private/Replication/Manager/ReplicationManagerState_Connected.h
@@ -34,7 +34,7 @@ namespace UE::ConcertSyncClient::Replication
 		FReplicationManagerState_Connected(
 			TSharedRef<IConcertClientSession> LiveSession,
 			IConcertClientReplicationBridge* ReplicationBridge,
-			TArray<FReplicationStreamDescription> StreamDescriptions,
+			TArray<FConcertReplicationStream> StreamDescriptions,
 			FReplicationManager& Owner
 			);
 		virtual ~FReplicationManagerState_Connected() override;
@@ -44,7 +44,7 @@ namespace UE::ConcertSyncClient::Replication
 		virtual void LeaveReplicationSession() override;
 		virtual bool CanJoin() override { return false; }
 		virtual bool IsConnectedToReplicationSession() override { return true; }
-		virtual EStreamEnumerationResult ForEachRegisteredStream(TFunctionRef<EBreakBehavior(const FReplicationStreamDescription& Stream)> Callback) const override;
+		virtual EStreamEnumerationResult ForEachRegisteredStream(TFunctionRef<EBreakBehavior(const FConcertReplicationStream& Stream)> Callback) const override;
 		virtual TFuture<FConcertReplication_ChangeAuthority_Response> RequestAuthorityChange(FConcertReplication_ChangeAuthority_Request Args) override;
 		virtual TFuture<FConcertReplication_QueryReplicationInfo_Response> QueryClientInfo(FConcertReplication_QueryReplicationInfo_Request Args) override;
 		virtual TFuture<FConcertReplication_ChangeStream_Response> ChangeStream(FConcertReplication_ChangeStream_Request Args) override;
@@ -59,7 +59,7 @@ namespace UE::ConcertSyncClient::Replication
 		/** Passed to FReplicationManagerState_Disconnected */
 		IConcertClientReplicationBridge* const ReplicationBridge;
 		/** The streams this client has registered with the server. */
-		TArray<FReplicationStreamDescription> RegisteredStreams;
+		TArray<FConcertReplicationStream> RegisteredStreams;
 		
 		/** The format this client will use for sending & receiving data. */
 		const TSharedRef<ConcertSyncCore::IObjectReplicationFormat> ReplicationFormat;
@@ -115,6 +115,6 @@ namespace UE::ConcertSyncClient::Replication
 		void RevertReleasingReplicatedObjects(const FConcertReplication_ChangeAuthority_Request& Request) const;
 		
 		/** Callback to Sender for obtaining an object's frequency settings. */
-		FConcertObjectReplicationSettings GetObjectFrequencySettings(const FReplicatedObjectId& Object) const;
+		FConcertObjectReplicationSettings GetObjectFrequencySettings(const FConcertReplicatedObjectId& Object) const;
 	};
 }
diff --git a/Engine/Plugins/Developer/Concert/ConcertSync/ConcertSyncClient/Source/ConcertSyncClient/Private/Replication/Processing/ClientReplicationDataCollector.cpp b/Engine/Plugins/Developer/Concert/ConcertSync/ConcertSyncClient/Source/ConcertSyncClient/Private/Replication/Processing/ClientReplicationDataCollector.cpp
--- a/Engine/Plugins/Developer/Concert/ConcertSync/ConcertSyncClient/Source/ConcertSyncClient/Private/Replication/Processing/ClientReplicationDataCollector.cpp
+++ b/Engine/Plugins/Developer/Concert/ConcertSync/ConcertSyncClient/Source/ConcertSyncClient/Private/Replication/Processing/ClientReplicationDataCollector.cpp
@@ -2,7 +2,7 @@
 
 #include "ClientReplicationDataCollector.h"
 
-#include "Replication/Data/ReplicationStreamDescription.h"
+#include "Replication/Data/ReplicationStream.h"
 #include "Replication/Formats/IObjectReplicationFormat.h"
 #include "Replication/IConcertClientReplicationBridge.h"
 #include "Replication/ReplicationPropertyFilter.h"
@@ -45,15 +45,15 @@ namespace UE::ConcertSyncClient::Replication
 		TArray<FObjectInfo>& ReplicatedObjectInfo = ObjectsToReplicate.FindOrAdd(Object);
 		const bool bIsNewReplicatedObject = ReplicatedObjectInfo.IsEmpty();
 		
-		const TArray<FReplicationStreamDescription>& RegisteredStreams = *GetStreamsDelegate.Execute();
+		const TArray<FConcertReplicationStream>& RegisteredStreams = *GetStreamsDelegate.Execute();
 		ReplicatedObjectInfo.Reserve(ReplicatedObjectInfo.Num() + AddedStreams.Num());
 		for (const FGuid& StreamId : AddedStreams)
 		{
-			const FReplicationStreamDescription* Stream = RegisteredStreams.FindByPredicate([&StreamId](const FReplicationStreamDescription& Stream)
+			const FConcertReplicationStream* Stream = RegisteredStreams.FindByPredicate([&StreamId](const FConcertReplicationStream& Stream)
 			{
 				return Stream.BaseDescription.Identifier == StreamId;
 			});
-			const FReplicatedObjectInfo* ObjectInfo = Stream ? Stream->BaseDescription.ReplicationMap.ReplicatedObjects.Find(Object) : nullptr;
+			const FConcertReplicatedObjectInfo* ObjectInfo = Stream ? Stream->BaseDescription.ReplicationMap.ReplicatedObjects.Find(Object) : nullptr;
 			if (ensureAlwaysMsgf(ObjectInfo, TEXT("Client's registered streams cache is out of sync")))
 			{
 				ReplicatedObjectInfo.Add({ StreamId, ObjectInfo->PropertySelection });
@@ -97,14 +97,14 @@ namespace UE::ConcertSyncClient::Replication
 			return;
 		}
 
-		const TArray<FReplicationStreamDescription>& RegisteredStreams = *GetStreamsDelegate.Execute();
+		const TArray<FConcertReplicationStream>& RegisteredStreams = *GetStreamsDelegate.Execute();
 		for (const FGuid& StreamId : PutStreams)
 		{
-			const FReplicationStreamDescription* Stream = RegisteredStreams.FindByPredicate([&StreamId](const FReplicationStreamDescription& Stream)
+			const FConcertReplicationStream* Stream = RegisteredStreams.FindByPredicate([&StreamId](const FConcertReplicationStream& Stream)
 			{
 				return Stream.BaseDescription.Identifier == StreamId;
 			});
-			const FReplicatedObjectInfo* ObjectInfo = Stream ? Stream->BaseDescription.ReplicationMap.ReplicatedObjects.Find(Object) : nullptr;
+			const FConcertReplicatedObjectInfo* ObjectInfo = Stream ? Stream->BaseDescription.ReplicationMap.ReplicatedObjects.Find(Object) : nullptr;
 			if (!ensureAlwaysMsgf(ObjectInfo, TEXT("Client's registered streams cache is out of sync")))
 			{
 				continue;
@@ -155,7 +155,7 @@ namespace UE::ConcertSyncClient::Replication
 		}
 	}
 
-	void FClientReplicationDataCollector::ForEachPendingObject(TFunctionRef<void(const FReplicatedObjectId&)> ProcessItemFunc) const
+	void FClientReplicationDataCollector::ForEachPendingObject(TFunctionRef<void(const FConcertReplicatedObjectId&)> ProcessItemFunc) const
 	{
 		for (const TPair<FSoftObjectPath, TArray<FObjectInfo>>& Pair : ObjectsToReplicate)
 		{
@@ -163,8 +163,8 @@ namespace UE::ConcertSyncClient::Replication
 			{
 				if (ObjectInfo.ObjectCache.IsValid())
 				{
-					ProcessItemFunc(FReplicatedObjectId {
-						FObjectInStreamID{ ObjectInfo.StreamId, ObjectInfo.ObjectCache.Get() },
+					ProcessItemFunc(FConcertReplicatedObjectId {
+						FConcertObjectInStreamID{ ObjectInfo.StreamId, ObjectInfo.ObjectCache.Get() },
 						ClientId
 					});
 				}
@@ -173,7 +173,7 @@ namespace UE::ConcertSyncClient::Replication
 	}
 
 	bool FClientReplicationDataCollector::ExtractReplicationDataForObject(
-		const FReplicatedObjectId& ObjectToProcess,
+		const FConcertReplicatedObjectId& ObjectToProcess,
 		TFunctionRef<void(const FConcertSessionSerializedPayload& Payload)> ProcessCopyable,
 		TFunctionRef<void(FConcertSessionSerializedPayload&& Payload)> ProcessMoveable
 		)
diff --git a/Engine/Plugins/Developer/Concert/ConcertSync/ConcertSyncClient/Source/ConcertSyncClient/Private/Replication/Processing/ClientReplicationDataCollector.h b/Engine/Plugins/Developer/Concert/ConcertSync/ConcertSyncClient/Source/ConcertSyncClient/Private/Replication/Processing/ClientReplicationDataCollector.h
--- a/Engine/Plugins/Developer/Concert/ConcertSync/ConcertSyncClient/Source/ConcertSyncClient/Private/Replication/Processing/ClientReplicationDataCollector.h
+++ b/Engine/Plugins/Developer/Concert/ConcertSync/ConcertSyncClient/Source/ConcertSyncClient/Private/Replication/Processing/ClientReplicationDataCollector.h
@@ -16,7 +16,7 @@ class IConcertClientReplicationBridge;
 class UObject;
 enum class EBreakBehavior : uint8;
 struct FConcertPropertySelection;
-struct FReplicationStreamDescription;
+struct FConcertReplicationStream;
 
 namespace UE::ConcertSyncCore
 {
@@ -34,7 +34,7 @@ namespace UE::ConcertSyncClient::Replication
 	{
 	public:
 
-		DECLARE_DELEGATE_RetVal(const TArray<FReplicationStreamDescription>*, FGetClientStreams);
+		DECLARE_DELEGATE_RetVal(const TArray<FConcertReplicationStream>*, FGetClientStreams);
 
 		/**
 		 * @param InReplicationBridge 
@@ -75,9 +75,9 @@ namespace UE::ConcertSyncClient::Replication
 		void AppendOwningStreamsForObject(const FSoftObjectPath& ObjectPath, TSet<FGuid>& Paths) const;
 
 		//~ Begin IReplicationDataSource Interface
-		virtual void ForEachPendingObject(TFunctionRef<void(const FReplicatedObjectId&)> ProcessItemFunc) const override;
+		virtual void ForEachPendingObject(TFunctionRef<void(const FConcertReplicatedObjectId&)> ProcessItemFunc) const override;
 		virtual int32 NumObjects() const override { return NumTrackedObjects; }
-		virtual bool ExtractReplicationDataForObject(const FReplicatedObjectId& Object, TFunctionRef<void(const FConcertSessionSerializedPayload& Payload)> ProcessCopyable, TFunctionRef<void(FConcertSessionSerializedPayload&& Payload)> ProcessMoveable) override;
+		virtual bool ExtractReplicationDataForObject(const FConcertReplicatedObjectId& Object, TFunctionRef<void(const FConcertSessionSerializedPayload& Payload)> ProcessCopyable, TFunctionRef<void(FConcertSessionSerializedPayload&& Payload)> ProcessMoveable) override;
 		//~ End IReplicationDataSource Interface
 
 	private:
diff --git a/Engine/Plugins/Developer/Concert/ConcertSync/ConcertSyncClient/Source/ConcertSyncClient/Private/Replication/Processing/ClientReplicationDataQueuer.cpp b/Engine/Plugins/Developer/Concert/ConcertSync/ConcertSyncClient/Source/ConcertSyncClient/Private/Replication/Processing/ClientReplicationDataQueuer.cpp
--- a/Engine/Plugins/Developer/Concert/ConcertSync/ConcertSyncClient/Source/ConcertSyncClient/Private/Replication/Processing/ClientReplicationDataQueuer.cpp
+++ b/Engine/Plugins/Developer/Concert/ConcertSync/ConcertSyncClient/Source/ConcertSyncClient/Private/Replication/Processing/ClientReplicationDataQueuer.cpp
@@ -20,7 +20,7 @@ namespace UE::ConcertSyncClient::Replication
 		: ReplicationBridge(ReplicationBridge)
 	{}
 
-	bool FClientReplicationDataQueuer::WantsToAcceptObject(const FReplicatedObjectId& Object) const
+	bool FClientReplicationDataQueuer::WantsToAcceptObject(const FConcertReplicatedObjectId& Object) const
 	{
 		// The server tries not to send objects to are unavailable but it could happen the object becomes unavailable while data is underway
 		const bool bIsAvailable = ReplicationBridge->IsObjectAvailable(Object.Object);
diff --git a/Engine/Plugins/Developer/Concert/ConcertSync/ConcertSyncClient/Source/ConcertSyncClient/Private/Replication/Processing/ClientReplicationDataQueuer.h b/Engine/Plugins/Developer/Concert/ConcertSync/ConcertSyncClient/Source/ConcertSyncClient/Private/Replication/Processing/ClientReplicationDataQueuer.h
--- a/Engine/Plugins/Developer/Concert/ConcertSync/ConcertSyncClient/Source/ConcertSyncClient/Private/Replication/Processing/ClientReplicationDataQueuer.h
+++ b/Engine/Plugins/Developer/Concert/ConcertSync/ConcertSyncClient/Source/ConcertSyncClient/Private/Replication/Processing/ClientReplicationDataQueuer.h
@@ -19,7 +19,7 @@ namespace UE::ConcertSyncClient::Replication
 		static TSharedRef<FClientReplicationDataQueuer> Make(IConcertClientReplicationBridge* ReplicationBridge, TSharedRef<ConcertSyncCore::FObjectReplicationCache> InReplicationCache);
 		
 		//~ Begin IReplicationCacheUser Interface
-		virtual bool WantsToAcceptObject(const FReplicatedObjectId& Object) const override;
+		virtual bool WantsToAcceptObject(const FConcertReplicatedObjectId& Object) const override;
 		//~ End IReplicationCacheUser Interface
 
 	private:
diff --git a/Engine/Plugins/Developer/Concert/ConcertSync/ConcertSyncClient/Source/ConcertSyncClient/Public/Replication/IConcertClientReplicationManager.h b/Engine/Plugins/Developer/Concert/ConcertSync/ConcertSyncClient/Source/ConcertSyncClient/Public/Replication/IConcertClientReplicationManager.h
--- a/Engine/Plugins/Developer/Concert/ConcertSync/ConcertSyncClient/Source/ConcertSyncClient/Public/Replication/IConcertClientReplicationManager.h
+++ b/Engine/Plugins/Developer/Concert/ConcertSync/ConcertSyncClient/Source/ConcertSyncClient/Public/Replication/IConcertClientReplicationManager.h
@@ -4,7 +4,6 @@
 
 #include "Misc/EBreakBehavior.h"
 #include "Replication/Messages/Handshake.h"
-#include "Replication/Data/ReplicationClientDescription.h"
 #include "Replication/Messages/ChangeAuthority.h"
 #include "Replication/Messages/ChangeStream.h"
 #include "Replication/Messages/ClientQuery.h"
@@ -12,17 +11,14 @@
 template<typename ResultType>
 class TFuture;
 
-struct FReplicationClientDescription;
-struct FReplicationStreamDescription;
+struct FConcertReplicationStream;
 
 namespace UE::ConcertSyncClient::Replication
 {
 	struct FJoinReplicatedSessionArgs
 	{
-		/** General info about this client, such as the type of data it wishes to receive. */
-		FReplicationClientDescription ClientInfo;
 		/** The streams this client offers. */
-		TArray<FReplicationStreamDescription> Streams;
+		TArray<FConcertReplicationStream> Streams;
 	};
 
 	struct FJoinReplicatedSessionResult
@@ -73,11 +69,11 @@ public:
 	 * @return Whether this manager is connected to a session (Iterated) or not (NoRegisteredStreams).
 	 * @note This future can finish on any thread (e.g. when message endpoint times out); usually it finishes on the game thread.  
 	 */
-	virtual EStreamEnumerationResult ForEachRegisteredStream(TFunctionRef<EBreakBehavior(const FReplicationStreamDescription& Stream)> Callback) const = 0;
+	virtual EStreamEnumerationResult ForEachRegisteredStream(TFunctionRef<EBreakBehavior(const FConcertReplicationStream& Stream)> Callback) const = 0;
 	/** @return Whether this manager is has any registered streams (basically whether ForEachRegisteredStream returns EStreamEnumerationResult::Iterated). */
 	bool HasRegisteredStreams() const;
 	/** @return The streams registered with the server. */
-	TArray<FReplicationStreamDescription> GetRegisteredStreams() const;
+	TArray<FConcertReplicationStream> GetRegisteredStreams() const;
 	
 	/**
 	 * Requests from the server to change the authority over some objects.
diff --git a/Engine/Plugins/Developer/Concert/ConcertSync/ConcertSyncCore/Source/ConcertSyncCore/Private/Replication/AuthorityConflictSharedUtils.cpp b/Engine/Plugins/Developer/Concert/ConcertSync/ConcertSyncCore/Source/ConcertSyncCore/Private/Replication/AuthorityConflictSharedUtils.cpp
--- a/Engine/Plugins/Developer/Concert/ConcertSync/ConcertSyncCore/Source/ConcertSyncCore/Private/Replication/AuthorityConflictSharedUtils.cpp
+++ b/Engine/Plugins/Developer/Concert/ConcertSync/ConcertSyncCore/Source/ConcertSyncCore/Private/Replication/AuthorityConflictSharedUtils.cpp
@@ -2,7 +2,7 @@
 
 #include "Replication/AuthorityConflictSharedUtils.h"
 
-#include "Replication/Data/ReplicationStreamDescription.h"
+#include "Replication/Data/ReplicationStream.h"
 #include "Replication/Messages/ChangeAuthority.h"
 #include "Replication/Messages/ChangeStream.h"
 
@@ -25,7 +25,7 @@ namespace UE::ConcertSyncCore::Replication::AuthorityConflictUtils
 				}
 				
 				EBreakBehavior Result = EBreakBehavior::Continue;
-				GroundTruth.ForEachStream(ClientEndpointId, [&Object, &Callback, &ClientEndpointId, &GroundTruth, &Result](const FGuid& StreamId, const FObjectReplicationMap& ReplicationMap) mutable
+				GroundTruth.ForEachStream(ClientEndpointId, [&Object, &Callback, &ClientEndpointId, &GroundTruth, &Result](const FGuid& StreamId, const FConcertObjectReplicationMap& ReplicationMap) mutable
 				{
 					// If client has not claimed authority over this object in this stream, skip
 					if (!GroundTruth.HasAuthority(ClientEndpointId, StreamId, Object))
@@ -34,8 +34,8 @@ namespace UE::ConcertSyncCore::Replication::AuthorityConflictUtils
 					}
 
 					// This client is using the request object: report the potential conflict ...
-					const FObjectReplicationMap& ObjectReplicationMap = ReplicationMap;
-					const FReplicatedObjectInfo* ReplicationObjectInfo = ObjectReplicationMap.ReplicatedObjects.Find(Object);
+					const FConcertObjectReplicationMap& ObjectReplicationMap = ReplicationMap;
+					const FConcertReplicatedObjectInfo* ReplicationObjectInfo = ObjectReplicationMap.ReplicatedObjects.Find(Object);
 					if (ReplicationObjectInfo && Callback(ClientEndpointId, StreamId, ReplicationObjectInfo->PropertySelection) == EBreakBehavior::Break)
 					{
 						// ... conflict ends iteration
@@ -87,13 +87,13 @@ namespace UE::ConcertSyncCore::Replication::AuthorityConflictUtils
 	void CleanseConflictsFromAuthorityRequest(FConcertReplication_ChangeAuthority_Request& Request, const FGuid& SendingClient, const IReplicationGroundTruth& GroundTruth)
 	{
 		// Need to check whether TakeAuthority is taking authority over properties other clients are already replicating
-		GroundTruth.ForEachStream(SendingClient, [&SendingClient, &Request, &GroundTruth](const FGuid& StreamId, const FObjectReplicationMap& ReplicationMap)
+		GroundTruth.ForEachStream(SendingClient, [&SendingClient, &Request, &GroundTruth](const FGuid& StreamId, const FConcertObjectReplicationMap& ReplicationMap)
 		{
 			for (auto ChangeIt = Request.TakeAuthority.CreateIterator(); ChangeIt; ++ChangeIt)
 			{
 				TPair<FSoftObjectPath, FConcertStreamArray>& Change = *ChangeIt;
 				
-				const FReplicatedObjectInfo* ObjectInfo = ReplicationMap.ReplicatedObjects.Find(Change.Key);
+				const FConcertReplicatedObjectInfo* ObjectInfo = ReplicationMap.ReplicatedObjects.Find(Change.Key);
 				if (!ObjectInfo)
 				{
 					ChangeIt.RemoveCurrent();
@@ -123,11 +123,11 @@ namespace UE::ConcertSyncCore::Replication::AuthorityConflictUtils
 	void CleanseConflictsFromStreamRequest(FConcertReplication_ChangeStream_Request& Request, const FGuid& SendingClient, const IReplicationGroundTruth& GroundTruth)
 	{
 		// Need to check whether ObjectsToPut adds any properties that an existing client has authority over.
-		GroundTruth.ForEachStream(SendingClient, [&SendingClient, &Request, &GroundTruth](const FGuid& StreamId, const FObjectReplicationMap& ReplicationMap)
+		GroundTruth.ForEachStream(SendingClient, [&SendingClient, &Request, &GroundTruth](const FGuid& StreamId, const FConcertObjectReplicationMap& ReplicationMap)
 		{
 			for (auto ChangeIt = Request.ObjectsToPut.CreateIterator(); ChangeIt; ++ChangeIt)
 			{
-				const TPair<FObjectInStreamID, FConcertReplication_ChangeStream_PutObject>& Change = *ChangeIt;
+				const TPair<FConcertObjectInStreamID, FConcertReplication_ChangeStream_PutObject>& Change = *ChangeIt;
 				const FConcertReplication_ChangeStream_PutObject& PutObject = Change.Value;
 
 				// Just changing the class?
diff --git a/Engine/Plugins/Developer/Concert/ConcertSync/ConcertSyncCore/Source/ConcertSyncCore/Private/Replication/ChangeStreamSharedUtils.cpp b/Engine/Plugins/Developer/Concert/ConcertSync/ConcertSyncCore/Source/ConcertSyncCore/Private/Replication/ChangeStreamSharedUtils.cpp
--- a/Engine/Plugins/Developer/Concert/ConcertSync/ConcertSyncCore/Source/ConcertSyncCore/Private/Replication/ChangeStreamSharedUtils.cpp
+++ b/Engine/Plugins/Developer/Concert/ConcertSync/ConcertSyncCore/Source/ConcertSyncCore/Private/Replication/ChangeStreamSharedUtils.cpp
@@ -11,18 +11,18 @@ namespace UE::ConcertSyncCore::Replication::ChangeStreamUtils
 {
 	void ForEachObjectLosingAuthority(
 		const FConcertReplication_ChangeStream_Request& Request,
-		const TArray<FReplicationStreamDescription>& ExistingStreams,
-		TFunctionRef<EBreakBehavior(const FObjectInStreamID&)> Callback
+		const TArray<FConcertReplicationStream>& ExistingStreams,
+		TFunctionRef<EBreakBehavior(const FConcertObjectInStreamID&)> Callback
 		)
 	{
-		for (const FReplicationStreamDescription& ExistingStream : ExistingStreams)
+		for (const FConcertReplicationStream& ExistingStream : ExistingStreams)
 		{
 			const FGuid& StreamId = ExistingStream.BaseDescription.Identifier;
 			if (Request.StreamsToRemove.Contains(StreamId))
 			{
-				for (const TPair<FSoftObjectPath, FReplicatedObjectInfo>& RemovePair : ExistingStream.BaseDescription.ReplicationMap.ReplicatedObjects)
+				for (const TPair<FSoftObjectPath, FConcertReplicatedObjectInfo>& RemovePair : ExistingStream.BaseDescription.ReplicationMap.ReplicatedObjects)
 				{
-					const FObjectInStreamID RemovedObjectId {  StreamId, RemovePair.Key };
+					const FConcertObjectInStreamID RemovedObjectId {  StreamId, RemovePair.Key };
 					if (Callback(RemovedObjectId) == EBreakBehavior::Break)
 					{
 						return;
@@ -33,7 +33,7 @@ namespace UE::ConcertSyncCore::Replication::ChangeStreamUtils
 
 		// This will call Callback multiple times if the same object is in ObjectsToRemove and StreamsToRemove
 		// but that does not really matter (it is a weird case anyhow - why would anyone construct such a request?)
-		for (const FObjectInStreamID& ObjectToRemove : Request.ObjectsToRemove)
+		for (const FConcertObjectInStreamID& ObjectToRemove : Request.ObjectsToRemove)
 		{
 			if (Callback(ObjectToRemove) == EBreakBehavior::Break)
 			{
@@ -42,16 +42,16 @@ namespace UE::ConcertSyncCore::Replication::ChangeStreamUtils
 		}
 	}
 	
-	void ApplyValidatedRequest(const FConcertReplication_ChangeStream_Request& Request, TArray<FReplicationStreamDescription>& StreamsToModify)
+	void ApplyValidatedRequest(const FConcertReplication_ChangeStream_Request& Request, TArray<FConcertReplicationStream>& StreamsToModify)
 	{
 		for (auto StreamIt = StreamsToModify.CreateIterator(); StreamIt; ++StreamIt)
 		{
-			FSharedReplicationStreamDescription& BaseDescription = StreamIt->BaseDescription;
-			TMap<FSoftObjectPath, FReplicatedObjectInfo>& ReplicationMap = BaseDescription.ReplicationMap.ReplicatedObjects;
+			FConcertBaseStreamInfo& BaseDescription = StreamIt->BaseDescription;
+			TMap<FSoftObjectPath, FConcertReplicatedObjectInfo>& ReplicationMap = BaseDescription.ReplicationMap.ReplicatedObjects;
 			FConcertStreamFrequencySettings& FrequencySettings = BaseDescription.FrequencySettings;
 			const FGuid& StreamId = BaseDescription.Identifier;
 			
-			for (const FObjectInStreamID& ObjectToRemove : Request.ObjectsToRemove)
+			for (const FConcertObjectInStreamID& ObjectToRemove : Request.ObjectsToRemove)
 			{
 				if (ObjectToRemove.StreamId == StreamId)
 				{
@@ -72,7 +72,7 @@ namespace UE::ConcertSyncCore::Replication::ChangeStreamUtils
 				StreamIt.RemoveCurrent();
 			}
 			
-			for (const TPair<FObjectInStreamID, FConcertReplication_ChangeStream_PutObject>& PutObjectPair : Request.ObjectsToPut)
+			for (const TPair<FConcertObjectInStreamID, FConcertReplication_ChangeStream_PutObject>& PutObjectPair : Request.ObjectsToPut)
 			{
 				if (PutObjectPair.Key.StreamId != StreamId)
 				{
@@ -81,7 +81,7 @@ namespace UE::ConcertSyncCore::Replication::ChangeStreamUtils
 
 				const FSoftClassPath& PathToSet = PutObjectPair.Value.ClassPath;
 				const FConcertPropertySelection& SelectionToSet = PutObjectPair.Value.Properties;
-				FReplicatedObjectInfo& ObjectInfo = ReplicationMap.FindOrAdd(PutObjectPair.Key.Object);
+				FConcertReplicatedObjectInfo& ObjectInfo = ReplicationMap.FindOrAdd(PutObjectPair.Key.Object);
 
 				// Write ClassPath if FindOrAdd just added it or if request specified a value to overwrite with ...
 				if (!ObjectInfo.ClassPath.IsValid() || PathToSet.IsValid())
@@ -101,7 +101,7 @@ namespace UE::ConcertSyncCore::Replication::ChangeStreamUtils
 			}
 		}
 
-		for (const FReplicationStreamDescription& StreamDescription : Request.StreamsToAdd)
+		for (const FConcertReplicationStream& StreamDescription : Request.StreamsToAdd)
 		{
 			StreamsToModify.Add(StreamDescription);
 		}
@@ -148,9 +148,9 @@ namespace UE::ConcertSyncCore::Replication::ChangeStreamUtils
 			}
 		}
 
-		static bool HasOrIsAddingProperties(const FConcertReplication_ChangeStream_Request& Request, const FSharedReplicationStreamDescription& Stream, const FSoftObjectPath& ObjectPath)
+		static bool HasOrIsAddingProperties(const FConcertReplication_ChangeStream_Request& Request, const FConcertBaseStreamInfo& Stream, const FSoftObjectPath& ObjectPath)
 		{
-			const FObjectReplicationMap& ReplicationMap = Stream.ReplicationMap;
+			const FConcertObjectReplicationMap& ReplicationMap = Stream.ReplicationMap;
 			return ReplicationMap.HasProperties(ObjectPath)
 				// It might be that ObjectsToPut is invalid.
 				// We do not need to check that here because the underlying FConcertReplication_ChangeStream_Request will be checked separately and would fail then anyways.
@@ -160,7 +160,7 @@ namespace UE::ConcertSyncCore::Replication::ChangeStreamUtils
 		static bool ValidateAddedFrequencies(
 			const FConcertReplication_ChangeStream_Request& FullRequest,
 			const TMap<FSoftObjectPath, FConcertObjectReplicationSettings>& Added,
-			const FSharedReplicationStreamDescription& Stream,
+			const FConcertBaseStreamInfo& Stream,
 			FConcertReplication_ChangeStream_FrequencyResponse* OptionalErrors = nullptr)
 		{
 			bool bAnyError = false;
@@ -183,7 +183,7 @@ namespace UE::ConcertSyncCore::Replication::ChangeStreamUtils
 
 	bool ValidateFrequencyChanges(
 		const FConcertReplication_ChangeStream_Request& Request,
-		const TArray<FReplicationStreamDescription>& Streams,
+		const TArray<FConcertReplicationStream>& Streams,
 		FConcertReplication_ChangeStream_FrequencyResponse* OptionalErrors
 		)
 	{
@@ -196,7 +196,7 @@ namespace UE::ConcertSyncCore::Replication::ChangeStreamUtils
 			const FConcertReplication_ChangeStream_Frequency& SubRequest = SubRequestPair.Value;
 			
 			const bool bIsModifyingDefaults = EnumHasAnyFlags(SubRequest.Flags, EConcertReplicationChangeFrequencyFlags::SetDefaults);
-			const FReplicationStreamDescription* StreamDescription = Streams.FindByPredicate([&TargetStreamId](const FReplicationStreamDescription& Description)
+			const FConcertReplicationStream* StreamDescription = Streams.FindByPredicate([&TargetStreamId](const FConcertReplicationStream& Description)
 			{
 				return Description.BaseDescription.Identifier == TargetStreamId;
 			});
@@ -230,11 +230,11 @@ namespace UE::ConcertSyncCore::Replication::ChangeStreamUtils
 #undef ADD_ERROR
 	
 	void IterateInvalidEntries(
-		const FObjectReplicationMap& ReplicationMap,
-		TFunctionRef<EBreakBehavior(const FSoftObjectPath&, const FReplicatedObjectInfo&)> Callback
+		const FConcertObjectReplicationMap& ReplicationMap,
+		TFunctionRef<EBreakBehavior(const FSoftObjectPath&, const FConcertReplicatedObjectInfo&)> Callback
 		)
 	{
-		for (const TPair<const FSoftObjectPath&, const FReplicatedObjectInfo&> Pair : ReplicationMap.ReplicatedObjects)
+		for (const TPair<const FSoftObjectPath&, const FConcertReplicatedObjectInfo&> Pair : ReplicationMap.ReplicatedObjects)
 		{
 			if (!Pair.Value.IsValidForSendingToServer() && Callback(Pair.Key, Pair.Value) == EBreakBehavior::Break)
 			{
@@ -245,16 +245,16 @@ namespace UE::ConcertSyncCore::Replication::ChangeStreamUtils
 
 	namespace Private
 	{
-		static void BuildPutObjectList(const FGuid& StreamId, const FObjectReplicationMap& Base, const FObjectReplicationMap& Desired, FConcertReplication_ChangeStream_Request& Request)
+		static void BuildPutObjectList(const FGuid& StreamId, const FConcertObjectReplicationMap& Base, const FConcertObjectReplicationMap& Desired, FConcertReplication_ChangeStream_Request& Request)
 		{
-			for (const TPair<FSoftObjectPath, FReplicatedObjectInfo>& BasePair : Base.ReplicatedObjects)
+			for (const TPair<FSoftObjectPath, FConcertReplicatedObjectInfo>& BasePair : Base.ReplicatedObjects)
 			{
 				const FSoftObjectPath& ObjectPath = BasePair.Key;
-				const FObjectInStreamID ObjectId { StreamId, ObjectPath };
-				const FReplicatedObjectInfo* DesiredObjectInfo = Desired.ReplicatedObjects.Find(ObjectPath);
+				const FConcertObjectInStreamID ObjectId { StreamId, ObjectPath };
+				const FConcertReplicatedObjectInfo* DesiredObjectInfo = Desired.ReplicatedObjects.Find(ObjectPath);
 				if (DesiredObjectInfo)
 				{
-					const FReplicatedObjectInfo& BaseObjectInfo = BasePair.Value;
+					const FConcertReplicatedObjectInfo& BaseObjectInfo = BasePair.Value;
 					const TOptional<FConcertReplication_ChangeStream_PutObject> PutObject = FConcertReplication_ChangeStream_PutObject::MakeFromChange(BaseObjectInfo, *DesiredObjectInfo);
 				
 					const bool bDesiredHasChangedFromBase = BaseObjectInfo != *DesiredObjectInfo;
@@ -272,9 +272,9 @@ namespace UE::ConcertSyncCore::Replication::ChangeStreamUtils
 			}
 		}
 
-		static void BuildRemoveObjectList(const FGuid& StreamId, const FObjectReplicationMap& Base, const FObjectReplicationMap& Desired, FConcertReplication_ChangeStream_Request& Request)
+		static void BuildRemoveObjectList(const FGuid& StreamId, const FConcertObjectReplicationMap& Base, const FConcertObjectReplicationMap& Desired, FConcertReplication_ChangeStream_Request& Request)
 		{
-			for (const TPair<FSoftObjectPath, FReplicatedObjectInfo>& DesiredPair : Desired.ReplicatedObjects)
+			for (const TPair<FSoftObjectPath, FConcertReplicatedObjectInfo>& DesiredPair : Desired.ReplicatedObjects)
 			{
 				const FSoftObjectPath& ObjectPath = DesiredPair.Key;
 				if (Base.ReplicatedObjects.Contains(ObjectPath))
@@ -284,7 +284,7 @@ namespace UE::ConcertSyncCore::Replication::ChangeStreamUtils
 				}
 
 				// Desired wants to add an object
-				const FReplicatedObjectInfo& ObjectInfo = DesiredPair.Value;
+				const FConcertReplicatedObjectInfo& ObjectInfo = DesiredPair.Value;
 				const TOptional<FConcertReplication_ChangeStream_PutObject> PutObject = FConcertReplication_ChangeStream_PutObject::MakeFromInfo(ObjectInfo);
 				const bool bDesiredStateHasEnoughDataToForPut = ensureMsgf(PutObject, TEXT("Function assumption violated; you did not pass in valid base or desired state."));
 				if (bDesiredStateHasEnoughDataToForPut)
@@ -297,8 +297,8 @@ namespace UE::ConcertSyncCore::Replication::ChangeStreamUtils
 
 	FConcertReplication_ChangeStream_Request BuildRequestFromDiff(
 		const FGuid& StreamId,
-		const FObjectReplicationMap& Base,
-		const FObjectReplicationMap& Desired
+		const FConcertObjectReplicationMap& Base,
+		const FConcertObjectReplicationMap& Desired
 		)
 	{
 		FConcertReplication_ChangeStream_Request Request;
diff --git a/Engine/Plugins/Developer/Concert/ConcertSync/ConcertSyncCore/Source/ConcertSyncCore/Private/Replication/Data/ObjectIds.cpp b/Engine/Plugins/Developer/Concert/ConcertSync/ConcertSyncCore/Source/ConcertSyncCore/Private/Replication/Data/ObjectIds.cpp
--- a/Engine/Plugins/Developer/Concert/ConcertSync/ConcertSyncCore/Source/ConcertSyncCore/Private/Replication/Data/ObjectIds.cpp
+++ b/Engine/Plugins/Developer/Concert/ConcertSync/ConcertSyncCore/Source/ConcertSyncCore/Private/Replication/Data/ObjectIds.cpp
@@ -2,12 +2,12 @@
 
 #include "Replication/Data/ObjectIds.h"
 
-uint32 GetTypeHash(const FObjectInStreamID& StreamObject)
+uint32 GetTypeHash(const FConcertObjectInStreamID& StreamObject)
 {
 	return HashCombineFast(GetTypeHash(StreamObject.Object), GetTypeHash(StreamObject.StreamId));
 }
 
-uint32 GetTypeHash(const FReplicatedObjectId& StreamObject)
+uint32 GetTypeHash(const FConcertReplicatedObjectId& StreamObject)
 {
-	return HashCombineFast(GetTypeHash(StreamObject.SenderEndpointId), GetTypeHash(static_cast<FObjectInStreamID>(StreamObject)));
+	return HashCombineFast(GetTypeHash(StreamObject.SenderEndpointId), GetTypeHash(static_cast<FConcertObjectInStreamID>(StreamObject)));
 }
\ No newline at end of file
diff --git a/Engine/Plugins/Developer/Concert/ConcertSync/ConcertSyncCore/Source/ConcertSyncCore/Private/Replication/Messages/ChangeStream.cpp b/Engine/Plugins/Developer/Concert/ConcertSync/ConcertSyncCore/Source/ConcertSyncCore/Private/Replication/Messages/ChangeStream.cpp
--- a/Engine/Plugins/Developer/Concert/ConcertSync/ConcertSyncCore/Source/ConcertSyncCore/Private/Replication/Messages/ChangeStream.cpp
+++ b/Engine/Plugins/Developer/Concert/ConcertSync/ConcertSyncCore/Source/ConcertSyncCore/Private/Replication/Messages/ChangeStream.cpp
@@ -4,7 +4,7 @@
 
 #include "Misc/OutputDevice.h"
 
-TOptional<FConcertReplication_ChangeStream_PutObject> FConcertReplication_ChangeStream_PutObject::MakeFromInfo(const FReplicatedObjectInfo& New)
+TOptional<FConcertReplication_ChangeStream_PutObject> FConcertReplication_ChangeStream_PutObject::MakeFromInfo(const FConcertReplicatedObjectInfo& New)
 {
 	if (!New.IsValidForSendingToServer())
 	{
@@ -13,7 +13,7 @@ TOptional<FConcertReplication_ChangeStream_PutObject> FConcertReplication_Change
 	return FConcertReplication_ChangeStream_PutObject{ New.PropertySelection, New.ClassPath };
 }
 
-TOptional<FConcertReplication_ChangeStream_PutObject> FConcertReplication_ChangeStream_PutObject::MakeFromChange(const FReplicatedObjectInfo& Base, const FReplicatedObjectInfo& Desired)
+TOptional<FConcertReplication_ChangeStream_PutObject> FConcertReplication_ChangeStream_PutObject::MakeFromChange(const FConcertReplicatedObjectInfo& Base, const FConcertReplicatedObjectInfo& Desired)
 {
 	if (Base.IsValidForSendingToServer() && Desired.IsValidForSendingToServer())
 	{
@@ -24,23 +24,23 @@ TOptional<FConcertReplication_ChangeStream_PutObject> FConcertReplication_Change
 	return {};
 }
 
-TOptional<FReplicatedObjectInfo> FConcertReplication_ChangeStream_PutObject::MakeObjectInfoIfValid() const
+TOptional<FConcertReplicatedObjectInfo> FConcertReplication_ChangeStream_PutObject::MakeObjectInfoIfValid() const
 {
 	if (Properties.ReplicatedProperties.IsEmpty() || ClassPath.IsNull())
 	{
 		return {};
 	}
-	return FReplicatedObjectInfo{ ClassPath, Properties };
+	return FConcertReplicatedObjectInfo{ ClassPath, Properties };
 }
 
 void FConcertReplication_ChangeStream_Response::LogErrors(FOutputDevice& OutputDevice) const
 {
-	for (const TPair<FObjectInStreamID, FReplicatedObjectId>& Conflict : AuthorityConflicts)
+	for (const TPair<FConcertObjectInStreamID, FConcertReplicatedObjectId>& Conflict : AuthorityConflicts)
 	{
 		OutputDevice.Logf(TEXT("Authority: { %s } conflicts with { %s }"), *Conflict.Key.ToString(), *Conflict.Value.ToString());
 	}
 
-	for (const TPair<FObjectInStreamID, EConcertPutObjectErrorCode>& Error : ObjectsToPutSemanticErrors)
+	for (const TPair<FConcertObjectInStreamID, EConcertPutObjectErrorCode>& Error : ObjectsToPutSemanticErrors)
 	{
 		const FString ErrorCodeAsString = [&Error]()
 		{
diff --git a/Engine/Plugins/Developer/Concert/ConcertSync/ConcertSyncCore/Source/ConcertSyncCore/Private/Replication/Processing/ObjectReplicationCache.cpp b/Engine/Plugins/Developer/Concert/ConcertSync/ConcertSyncCore/Source/ConcertSyncCore/Private/Replication/Processing/ObjectReplicationCache.cpp
--- a/Engine/Plugins/Developer/Concert/ConcertSync/ConcertSyncCore/Source/ConcertSyncCore/Private/Replication/Processing/ObjectReplicationCache.cpp
+++ b/Engine/Plugins/Developer/Concert/ConcertSync/ConcertSyncCore/Source/ConcertSyncCore/Private/Replication/Processing/ObjectReplicationCache.cpp
@@ -16,7 +16,7 @@ namespace UE::ConcertSyncCore
 	{
 		int32 NumAccepted = 0;
 		
-		const FReplicatedObjectId ObjectId{ { OriginStreamId, ObjectReplicationEvent.ReplicatedObject }, SendingEndpointId };
+		const FConcertReplicatedObjectId ObjectId{ { OriginStreamId, ObjectReplicationEvent.ReplicatedObject }, SendingEndpointId };
 		FObjectCache* ObjectCacheBeforeAddition = Cache.Find(ObjectId);
 		if (ObjectCacheBeforeAddition)
 		{
@@ -66,7 +66,7 @@ namespace UE::ConcertSyncCore
 					}
 
 					// ObjectCache may not be found because cache user was unregistered and then destroyed: UnregisterDataCacheUser removes cache users.
-					const FObjectInStreamID ObjectId{ OriginStreamId, LazilyCopiedEventPtr->ReplicatedObject };
+					const FConcertObjectInStreamID ObjectId{ OriginStreamId, LazilyCopiedEventPtr->ReplicatedObject };
 					FObjectCache* ObjectCache = This->Cache.Find(ObjectId);
 					if (!ObjectCache)
 					{
diff --git a/Engine/Plugins/Developer/Concert/ConcertSync/ConcertSyncCore/Source/ConcertSyncCore/Private/Replication/Processing/ObjectReplicationProcessor.cpp b/Engine/Plugins/Developer/Concert/ConcertSync/ConcertSyncCore/Source/ConcertSyncCore/Private/Replication/Processing/ObjectReplicationProcessor.cpp
--- a/Engine/Plugins/Developer/Concert/ConcertSync/ConcertSyncCore/Source/ConcertSyncCore/Private/Replication/Processing/ObjectReplicationProcessor.cpp
+++ b/Engine/Plugins/Developer/Concert/ConcertSync/ConcertSyncCore/Source/ConcertSyncCore/Private/Replication/Processing/ObjectReplicationProcessor.cpp
@@ -16,7 +16,7 @@ namespace UE::ConcertSyncCore
 	void FObjectReplicationProcessor::ProcessObjects(const FProcessObjectsParams& Params)
 	{
 		// TODO UE-190714: Respect time budget and prioritize objects
-		DataSource->ForEachPendingObject([this](const FReplicatedObjectId& ObjectInfo)
+		DataSource->ForEachPendingObject([this](const FConcertReplicatedObjectId& ObjectInfo)
 		{
 			ProcessObject({ ObjectInfo });
 		});
diff --git a/Engine/Plugins/Developer/Concert/ConcertSync/ConcertSyncCore/Source/ConcertSyncCore/Private/Replication/Processing/Proxy/FrequencyTracker.cpp b/Engine/Plugins/Developer/Concert/ConcertSync/ConcertSyncCore/Source/ConcertSyncCore/Private/Replication/Processing/Proxy/FrequencyTracker.cpp
--- a/Engine/Plugins/Developer/Concert/ConcertSync/ConcertSyncCore/Source/ConcertSyncCore/Private/Replication/Processing/Proxy/FrequencyTracker.cpp
+++ b/Engine/Plugins/Developer/Concert/ConcertSync/ConcertSyncCore/Source/ConcertSyncCore/Private/Replication/Processing/Proxy/FrequencyTracker.cpp
@@ -7,7 +7,7 @@
 namespace UE::ConcertSyncCore
 {
 	ECanProcessObject FFrequencyTracker::TrackAndCheckObject(
-		const FObjectInStreamID& ObjectPath,
+		const FConcertObjectInStreamID& ObjectPath,
 		const FConcertObjectReplicationSettings& ReplicationSettings,
 		const FPlatformSecondsTimestamp NowAsPlatformSeconds
 		)
@@ -17,7 +17,7 @@ namespace UE::ConcertSyncCore
 		return ReplicationSettings.CanProcessObject(LastTimeProcessed, NowAsPlatformSeconds);
 	}
 
-	void FFrequencyTracker::OnProcessObject(const FObjectInStreamID& Object, const FPlatformSecondsTimestamp NowAsPlatformSeconds)
+	void FFrequencyTracker::OnProcessObject(const FConcertObjectInStreamID& Object, const FPlatformSecondsTimestamp NowAsPlatformSeconds)
 	{
 		ObjectsToLastProcessed.FindOrAdd(Object) = NowAsPlatformSeconds;
 	}
diff --git a/Engine/Plugins/Developer/Concert/ConcertSync/ConcertSyncCore/Source/ConcertSyncCore/Private/Replication/Processing/ReplicationDataQueuer.cpp b/Engine/Plugins/Developer/Concert/ConcertSync/ConcertSyncCore/Source/ConcertSyncCore/Private/Replication/Processing/ReplicationDataQueuer.cpp
--- a/Engine/Plugins/Developer/Concert/ConcertSync/ConcertSyncCore/Source/ConcertSyncCore/Private/Replication/Processing/ReplicationDataQueuer.cpp
+++ b/Engine/Plugins/Developer/Concert/ConcertSync/ConcertSyncCore/Source/ConcertSyncCore/Private/Replication/Processing/ReplicationDataQueuer.cpp
@@ -14,7 +14,7 @@ namespace UE::ConcertSyncCore
 		}
 	}
 
-	void FReplicationDataQueuer::ForEachPendingObject(TFunctionRef<void(const FReplicatedObjectId&)> ProcessItemFunc) const
+	void FReplicationDataQueuer::ForEachPendingObject(TFunctionRef<void(const FConcertReplicatedObjectId&)> ProcessItemFunc) const
 	{
 		for (auto It = PendingEvents.CreateConstIterator(); It; ++It)
 		{
@@ -28,7 +28,7 @@ namespace UE::ConcertSyncCore
 	}
 
 	bool FReplicationDataQueuer::ExtractReplicationDataForObject(
-		const FReplicatedObjectId& Object,
+		const FConcertReplicatedObjectId& Object,
 		TFunctionRef<void(const FConcertSessionSerializedPayload& Payload)> ProcessCopyable,
 		TFunctionRef<void(FConcertSessionSerializedPayload&& Payload)> ProcessMoveable
 		)
@@ -45,7 +45,7 @@ namespace UE::ConcertSyncCore
 		return true;
 	}
 
-	void FReplicationDataQueuer::OnDataCached(const FReplicatedObjectId& Object, TSharedRef<const FConcertReplication_ObjectReplicationEvent> Data)
+	void FReplicationDataQueuer::OnDataCached(const FConcertReplicatedObjectId& Object, TSharedRef<const FConcertReplication_ObjectReplicationEvent> Data)
 	{
 		PendingEvents.Add(Object, MoveTemp(Data));
 	}
diff --git a/Engine/Plugins/Developer/Concert/ConcertSync/ConcertSyncCore/Source/ConcertSyncCore/Public/Replication/AuthorityConflictSharedUtils.h b/Engine/Plugins/Developer/Concert/ConcertSync/ConcertSyncCore/Source/ConcertSyncCore/Public/Replication/AuthorityConflictSharedUtils.h
--- a/Engine/Plugins/Developer/Concert/ConcertSync/ConcertSyncCore/Source/ConcertSyncCore/Public/Replication/AuthorityConflictSharedUtils.h
+++ b/Engine/Plugins/Developer/Concert/ConcertSync/ConcertSyncCore/Source/ConcertSyncCore/Public/Replication/AuthorityConflictSharedUtils.h
@@ -12,8 +12,8 @@ struct FConcertPropertyChain;
 struct FConcertPropertySelection;
 struct FConcertReplication_ChangeAuthority_Request;
 struct FConcertReplication_ChangeStream_Request;
-struct FObjectReplicationMap;
-struct FReplicatedObjectId;
+struct FConcertObjectReplicationMap;
+struct FConcertReplicatedObjectId;
 
 namespace UE::ConcertSyncCore::Replication::AuthorityConflictUtils
 {
@@ -33,7 +33,7 @@ namespace UE::ConcertSyncCore::Replication::AuthorityConflictUtils
 	public:
 
 		/** Provides a way to extract all streams registered to a given client. */
-		virtual void ForEachStream(const FGuid& ClientEndpointId, TFunctionRef<EBreakBehavior(const FGuid& StreamId, const FObjectReplicationMap& ReplicationMap)> Callback) const = 0;
+		virtual void ForEachStream(const FGuid& ClientEndpointId, TFunctionRef<EBreakBehavior(const FGuid& StreamId, const FConcertObjectReplicationMap& ReplicationMap)> Callback) const = 0;
 
 		/** Iterates through all clients have registered to send any data. */
 		virtual void ForEachSendingClient(TFunctionRef<EBreakBehavior(const FGuid& ClientEndpointId)> Callback) const = 0;
diff --git a/Engine/Plugins/Developer/Concert/ConcertSync/ConcertSyncCore/Source/ConcertSyncCore/Public/Replication/ChangeStreamSharedUtils.h b/Engine/Plugins/Developer/Concert/ConcertSync/ConcertSyncCore/Source/ConcertSyncCore/Public/Replication/ChangeStreamSharedUtils.h
--- a/Engine/Plugins/Developer/Concert/ConcertSync/ConcertSyncCore/Source/ConcertSyncCore/Public/Replication/ChangeStreamSharedUtils.h
+++ b/Engine/Plugins/Developer/Concert/ConcertSync/ConcertSyncCore/Source/ConcertSyncCore/Public/Replication/ChangeStreamSharedUtils.h
@@ -10,10 +10,10 @@
 enum class EBreakBehavior : uint8;
 
 struct FConcertReplication_ChangeStream_Request;
-struct FObjectInStreamID;
-struct FObjectReplicationMap;
-struct FReplicationStreamDescription;
-struct FReplicatedObjectInfo;
+struct FConcertObjectInStreamID;
+struct FConcertObjectReplicationMap;
+struct FConcertReplicationStream;
+struct FConcertReplicatedObjectInfo;
 
 /** This namespace contains the shared logic for applying FConcertChangeStream_Requests. */
 namespace UE::ConcertSyncCore::Replication::ChangeStreamUtils
@@ -27,8 +27,8 @@ namespace UE::ConcertSyncCore::Replication::ChangeStreamUtils
 	 */
 	CONCERTSYNCCORE_API void ForEachObjectLosingAuthority(
 		const FConcertReplication_ChangeStream_Request& Request,
-		const TArray<FReplicationStreamDescription>& ExistingStreams,
-		TFunctionRef<EBreakBehavior(const FObjectInStreamID&)> Callback
+		const TArray<FConcertReplicationStream>& ExistingStreams,
+		TFunctionRef<EBreakBehavior(const FConcertObjectInStreamID&)> Callback
 		);
 	
 	/**
@@ -38,7 +38,7 @@ namespace UE::ConcertSyncCore::Replication::ChangeStreamUtils
 	 */
 	CONCERTSYNCCORE_API void ApplyValidatedRequest(
 		const FConcertReplication_ChangeStream_Request& Request,
-		IN OUT TArray<FReplicationStreamDescription>& StreamsToModify
+		IN OUT TArray<FConcertReplicationStream>& StreamsToModify
 		);
 	
 	/**
@@ -65,12 +65,12 @@ namespace UE::ConcertSyncCore::Replication::ChangeStreamUtils
 	 */
 	CONCERTSYNCCORE_API bool ValidateFrequencyChanges(
 		const FConcertReplication_ChangeStream_Request& Request,
-		const TArray<FReplicationStreamDescription>& Streams,
+		const TArray<FConcertReplicationStream>& Streams,
 		OUT FConcertReplication_ChangeStream_FrequencyResponse* OptionalErrors = nullptr
 		);
 
 	/** Iterates entries with incomplete data, which are entries for which IsValidForSendingToServer returns false. */
-	CONCERTSYNCCORE_API void IterateInvalidEntries(const FObjectReplicationMap& ReplicationMap, TFunctionRef<EBreakBehavior(const FSoftObjectPath&, const FReplicatedObjectInfo&)> Callback);
+	CONCERTSYNCCORE_API void IterateInvalidEntries(const FConcertObjectReplicationMap& ReplicationMap, TFunctionRef<EBreakBehavior(const FSoftObjectPath&, const FConcertReplicatedObjectInfo&)> Callback);
 
 	/**
 	 * Builds PutObject and RemoveObject requests that would transform Base to Desired.
@@ -87,7 +87,7 @@ namespace UE::ConcertSyncCore::Replication::ChangeStreamUtils
 	 */
 	CONCERTSYNCCORE_API FConcertReplication_ChangeStream_Request BuildRequestFromDiff(
 		const FGuid& StreamId,
-		const FObjectReplicationMap& Base,
-		const FObjectReplicationMap& Desired
+		const FConcertObjectReplicationMap& Base,
+		const FConcertObjectReplicationMap& Desired
 		);
 }
diff --git a/Engine/Plugins/Developer/Concert/ConcertSync/ConcertSyncCore/Source/ConcertSyncCore/Public/Replication/Data/ClientQueriedInfo.h b/Engine/Plugins/Developer/Concert/ConcertSync/ConcertSyncCore/Source/ConcertSyncCore/Public/Replication/Data/ClientQueriedInfo.h
--- a/Engine/Plugins/Developer/Concert/ConcertSync/ConcertSyncCore/Source/ConcertSyncCore/Public/Replication/Data/ClientQueriedInfo.h
+++ b/Engine/Plugins/Developer/Concert/ConcertSync/ConcertSyncCore/Source/ConcertSyncCore/Public/Replication/Data/ClientQueriedInfo.h
@@ -2,12 +2,12 @@
 
 #pragma once
 
-#include "Replication/Data/ReplicationStreamDescription.h"
+#include "ReplicationStream.h"
 #include "ClientQueriedInfo.generated.h"
 
 /** Describes objects a client has authority over */
 USTRUCT()
-struct FReplicationAuthorityInfo
+struct FConcertAuthorityClientInfo
 {
 	GENERATED_BODY()
 
@@ -25,7 +25,7 @@ struct FReplicationAuthorityInfo
 
 /** This is info a client receives about another client via FConcertQueryClientStreams_Request. */
 USTRUCT()
-struct FReplicationClientQueriedInfo
+struct FConcertQueriedClientInfo
 {
 	GENERATED_BODY()
 
@@ -33,11 +33,11 @@ struct FReplicationClientQueriedInfo
 	 * The streams the client has registered.
 	 * 
 	 * If EConcertQueryClientStreamFlags::SkipStreamInfo was set, this is empty.
-	 * If EConcertQueryClientStreamFlags::SkipProperties was set, streams' FReplicatedObjectInfo::PropertySelection are empty.
+	 * If EConcertQueryClientStreamFlags::SkipProperties was set, streams' FConcertReplicatedObjectInfo::PropertySelection are empty.
 	 * If EConcertQueryClientStreamFlags::SkipFrequency was set, the FrequencySettings are empty.
 	 */
 	UPROPERTY()
-	TArray<FSharedReplicationStreamDescription> Streams;
+	TArray<FConcertBaseStreamInfo> Streams;
 
 	/**
 	 * Indirectly describes which object properties the client has authority over.
@@ -45,5 +45,5 @@ struct FReplicationClientQueriedInfo
 	 * If EConcertQueryClientStreamFlags::SkipAuthority was set, this is empty.
 	 */
 	UPROPERTY()
-	TArray<FReplicationAuthorityInfo> Authority;
+	TArray<FConcertAuthorityClientInfo> Authority;
 };
\ No newline at end of file
diff --git a/Engine/Plugins/Developer/Concert/ConcertSync/ConcertSyncCore/Source/ConcertSyncCore/Public/Replication/Data/ObjectIds.h b/Engine/Plugins/Developer/Concert/ConcertSync/ConcertSyncCore/Source/ConcertSyncCore/Public/Replication/Data/ObjectIds.h
--- a/Engine/Plugins/Developer/Concert/ConcertSync/ConcertSyncCore/Source/ConcertSyncCore/Public/Replication/Data/ObjectIds.h
+++ b/Engine/Plugins/Developer/Concert/ConcertSync/ConcertSyncCore/Source/ConcertSyncCore/Public/Replication/Data/ObjectIds.h
@@ -8,7 +8,7 @@
 
 /** Utility for identifying an object part of a replication stream definition.*/
 USTRUCT()
-struct FObjectInStreamID
+struct FConcertObjectInStreamID
 {
 	GENERATED_BODY()
 	
@@ -21,12 +21,12 @@ struct FObjectInStreamID
 
 	FString ToString() const { return FString::Printf(TEXT("StreamId: %s, Object: %s"), *StreamId.ToString(), *Object.ToString()); }
 		
-	friend bool operator==(const FObjectInStreamID& Left, const FObjectInStreamID& Right)
+	friend bool operator==(const FConcertObjectInStreamID& Left, const FConcertObjectInStreamID& Right)
 	{
 		return Left.StreamId == Right.StreamId && Left.Object == Right.Object;
 	}
 
-	friend bool operator!=(const FObjectInStreamID& Left, const FObjectInStreamID& Right)
+	friend bool operator!=(const FConcertObjectInStreamID& Left, const FConcertObjectInStreamID& Right)
 	{
 		return !(Left == Right);
 	}
@@ -34,7 +34,7 @@ struct FObjectInStreamID
 
 /** Identifies an object that was replicated by a client based on an underlying stream. */
 USTRUCT()
-struct FReplicatedObjectId : public FObjectInStreamID
+struct FConcertReplicatedObjectId : public FConcertObjectInStreamID
 {
 	GENERATED_BODY()
 	
@@ -43,17 +43,17 @@ struct FReplicatedObjectId : public FObjectInStreamID
 	
 	FString ToString() const { return FString::Printf(TEXT("StreamId: %s, Object: %s, Sender: %s"), *StreamId.ToString(), *Object.ToString(), *SenderEndpointId.ToString()); }
 		
-	friend bool operator==(const FReplicatedObjectId& Left, const FReplicatedObjectId& Right)
+	friend bool operator==(const FConcertReplicatedObjectId& Left, const FConcertReplicatedObjectId& Right)
 	{
 		return Left.SenderEndpointId == Right.SenderEndpointId
-			&& static_cast<const FObjectInStreamID&>(Left) == static_cast<const FObjectInStreamID&>(Right);
+			&& static_cast<const FConcertObjectInStreamID&>(Left) == static_cast<const FConcertObjectInStreamID&>(Right);
 	}
 
-	friend bool operator!=(const FReplicatedObjectId& Left, const FReplicatedObjectId& Right)
+	friend bool operator!=(const FConcertReplicatedObjectId& Left, const FConcertReplicatedObjectId& Right)
 	{
 		return !(Left == Right);
 	}
 };
 
-CONCERTSYNCCORE_API uint32 GetTypeHash(const FObjectInStreamID& StreamObject);
-CONCERTSYNCCORE_API uint32 GetTypeHash(const FReplicatedObjectId& StreamObject);
+CONCERTSYNCCORE_API uint32 GetTypeHash(const FConcertObjectInStreamID& StreamObject);
+CONCERTSYNCCORE_API uint32 GetTypeHash(const FConcertReplicatedObjectId& StreamObject);
diff --git a/Engine/Plugins/Developer/Concert/ConcertSync/ConcertSyncCore/Source/ConcertSyncCore/Public/Replication/Data/ObjectReplicationMap.h b/Engine/Plugins/Developer/Concert/ConcertSync/ConcertSyncCore/Source/ConcertSyncCore/Public/Replication/Data/ObjectReplicationMap.h
--- a/Engine/Plugins/Developer/Concert/ConcertSync/ConcertSyncCore/Source/ConcertSyncCore/Public/Replication/Data/ObjectReplicationMap.h
+++ b/Engine/Plugins/Developer/Concert/ConcertSync/ConcertSyncCore/Source/ConcertSyncCore/Public/Replication/Data/ObjectReplicationMap.h
@@ -7,7 +7,7 @@
 #include "ObjectReplicationMap.generated.h"
 
 USTRUCT()
-struct FReplicatedObjectInfo
+struct FConcertReplicatedObjectInfo
 {
 	GENERATED_BODY()
 
@@ -20,7 +20,7 @@ struct FReplicatedObjectInfo
 	FConcertPropertySelection PropertySelection;
 
 	// Use static factories instead of user-provided constructors to continue allowing brace-initialization
-	static FReplicatedObjectInfo Make(const UObject& Object)
+	static FConcertReplicatedObjectInfo Make(const UObject& Object)
 	{
 		return { Object.GetClass() };
 	}
@@ -32,12 +32,12 @@ struct FReplicatedObjectInfo
 			&& !PropertySelection.ReplicatedProperties.IsEmpty();
 	}
 	
-	friend bool operator==(const FReplicatedObjectInfo& Left, const FReplicatedObjectInfo& Right)
+	friend bool operator==(const FConcertReplicatedObjectInfo& Left, const FConcertReplicatedObjectInfo& Right)
 	{
 		return Left.ClassPath == Right.ClassPath
 			&& Left.PropertySelection == Right.PropertySelection;
 	}
-	friend bool operator!=(const FReplicatedObjectInfo& Left, const FReplicatedObjectInfo& Right)
+	friend bool operator!=(const FConcertReplicatedObjectInfo& Left, const FConcertReplicatedObjectInfo& Right)
 	{
 		return !(Left == Right);
 	}
@@ -45,7 +45,7 @@ struct FReplicatedObjectInfo
 
 /** Maps objects to their replicated properties. */
 USTRUCT()
-struct FObjectReplicationMap
+struct FConcertObjectReplicationMap
 {
 	GENERATED_BODY()
 
@@ -56,12 +56,12 @@ struct FObjectReplicationMap
 	 * Technically, this can also be non-UWorld objects.
 	 */
 	UPROPERTY()
-	TMap<FSoftObjectPath, FReplicatedObjectInfo> ReplicatedObjects;
+	TMap<FSoftObjectPath, FConcertReplicatedObjectInfo> ReplicatedObjects;
 
 	/** @return Whether ObjectPath has any properties assigned to it. */
 	bool HasProperties(const FSoftObjectPath& ObjectPath) const
 	{
-		const FReplicatedObjectInfo* ObjectInfo = ReplicatedObjects.Find(ObjectPath);
+		const FConcertReplicatedObjectInfo* ObjectInfo = ReplicatedObjects.Find(ObjectPath);
 		return ObjectInfo && !ObjectInfo->PropertySelection.ReplicatedProperties.IsEmpty();
 	}
 
@@ -71,11 +71,11 @@ struct FObjectReplicationMap
 		return ReplicatedObjects.IsEmpty();
 	}
 
-	friend bool operator==(const FObjectReplicationMap& Left, const FObjectReplicationMap& Right)
+	friend bool operator==(const FConcertObjectReplicationMap& Left, const FConcertObjectReplicationMap& Right)
 	{
 		return Left.ReplicatedObjects.OrderIndependentCompareEqual(Right.ReplicatedObjects);
 	}
-	friend bool operator!=(const FObjectReplicationMap& Left, const FObjectReplicationMap& Right)
+	friend bool operator!=(const FConcertObjectReplicationMap& Left, const FConcertObjectReplicationMap& Right)
 	{
 		return !(Left == Right);
 	}
diff --git a/Engine/Plugins/Developer/Concert/ConcertSync/ConcertSyncCore/Source/ConcertSyncCore/Public/Replication/Data/ReplicationClientDescription.h b/Engine/Plugins/Developer/Concert/ConcertSync/ConcertSyncCore/Source/ConcertSyncCore/Public/Replication/Data/ReplicationClientDescription.h
deleted file mode 100644
diff --git a/Engine/Plugins/Developer/Concert/ConcertSync/ConcertSyncCore/Source/ConcertSyncCore/Public/Replication/Data/ReplicationFrequencySettings.h b/Engine/Plugins/Developer/Concert/ConcertSync/ConcertSyncCore/Source/ConcertSyncCore/Public/Replication/Data/ReplicationFrequencySettings.h
--- a/Engine/Plugins/Developer/Concert/ConcertSync/ConcertSyncCore/Source/ConcertSyncCore/Public/Replication/Data/ReplicationFrequencySettings.h
+++ b/Engine/Plugins/Developer/Concert/ConcertSync/ConcertSyncCore/Source/ConcertSyncCore/Public/Replication/Data/ReplicationFrequencySettings.h
@@ -8,7 +8,7 @@
 UENUM()
 enum class EConcertObjectReplicationMode : uint8
 {
-	/** Replicate at the rate specified at FReplicatedObjectInfo::ReplicationRate */
+	/** Replicate at the rate specified at FConcertReplicatedObjectInfo::ReplicationRate */
 	SpecifiedRate,
 	/** Replicate the object as often as possible: every tick. */
 	Realtime
diff --git a/Engine/Plugins/Developer/Concert/ConcertSync/ConcertSyncCore/Source/ConcertSyncCore/Public/Replication/Data/ReplicationStream.h b/Engine/Plugins/Developer/Concert/ConcertSync/ConcertSyncCore/Source/ConcertSyncCore/Public/Replication/Data/ReplicationStream.h
new file mode 100644
--- /dev/null
+++ b/Engine/Plugins/Developer/Concert/ConcertSync/ConcertSyncCore/Source/ConcertSyncCore/Public/Replication/Data/ReplicationStream.h
@@ -0,0 +1,137 @@
+// Copyright Epic Games, Inc. All Rights Reserved.
+
+#pragma once
+
+#include "ObjectReplicationMap.h"
+#include "ReplicationFrequencySettings.h"
+#include "ReplicationStream.generated.h"
+
+template<typename OptionalType>
+struct TOptional;
+
+namespace UE::ConcertSyncCore
+{
+	enum class EReplicationStreamCloneFlags : uint8
+	{
+		None,
+		SkipProperties = 1 << 0,
+		SkipFrequency = 1 << 1
+	};
+	ENUM_CLASS_FLAGS(EReplicationStreamCloneFlags)
+}
+
+/**
+ * Holds the base info of a replication stream: it knows the data that was registered to be sent and at what frequency.
+ * 
+ * This struct exists to future proof the change UE-190167, which will add UObject rule objects to streams.
+ * When that happens, this struct holds all non-UObject data for the stream.
+ */
+USTRUCT()
+struct FConcertBaseStreamInfo
+{
+	GENERATED_BODY()
+	
+	/** Unique id for this stream. */
+	UPROPERTY()
+	FGuid Identifier;
+
+	/** Identifies the data that this stream will send. */
+	UPROPERTY()
+	FConcertObjectReplicationMap ReplicationMap;
+
+	/**
+	 * Determines how often updates are sent for objects.
+	 *
+	 * The purposes of registering frequency with the server communicates the sending client's intent and has these advantages:
+	 * - If the server receives replicated data more often than specified here, it will combine and throttle the events forwarded to the other clients.
+	 * - The receiving clients can query how often per frame to expect updates for an object
+	 *
+	 * On a technical level, this field is not required for replicating objects: in a hypothetical, alternate implementation, the sending client could
+	 * just keep the replication frequency local to themselves and send at the rate it wants to send.
+	 * This field is here mostly to allow inspection by other parties, which does add overhead (engineering and runtime) because we must support
+	 * FConcertReplication_ChangeStream_Frequency requests).
+	 * 
+	 * With that said, the sending implementation uses the rate that is specified here to throttle the rate at which it sends update to the server.
+	 * @see FObjectReplicationProcessor::ProcessObjects.
+	 */
+	UPROPERTY()
+	FConcertStreamFrequencySettings FrequencySettings;
+
+	/** Util for efficiently cloning stream info and optionally skipping some info. */
+	FConcertBaseStreamInfo Clone(UE::ConcertSyncCore::EReplicationStreamCloneFlags Flags = UE::ConcertSyncCore::EReplicationStreamCloneFlags::None) const;
+	
+	friend bool operator==(const FConcertBaseStreamInfo& Left, const FConcertBaseStreamInfo& Right)
+	{
+		return Left.Identifier == Right.Identifier
+			&& Left.ReplicationMap == Right.ReplicationMap
+			&& Left.FrequencySettings == Right.FrequencySettings;
+	}
+	friend bool operator!=(const FConcertBaseStreamInfo& Left, const FConcertBaseStreamInfo& Right)
+	{
+		return !(Left == Right);
+	}
+};
+
+/**
+ * A replication stream a related set of data that is to be replicated. The application using this frame work (e.g. Multi User) defines what it means
+ * for data to be "related".
+ *
+ * The stream defines:
+ * - a map of objects to their replicated properties,
+ * - at what frequencies the data replicates,
+ * - a set of stream attributes which receiving clients can use to tell the server what data they want to receive (rules are TODO UE-190167)
+ * 
+ * Streams exist mostly for logical grouping and organization. It is completely acceptable for an application to only register one stream per client (Multi User does this).
+ * However, there is a minor technical detail: While streams registered by the same client can replicate the same UObjects (and even overlap properties), streams send data
+ * independently. The consequence is that each stream:
+ * - has its own frequency counter
+ * - streams data is not combined / merged (i.e. there is an overhead of the book-keeping data that is sent with each update).
+ *
+ * This struct only contains FConcertBaseStreamInfo for to future proof the change UE-190167, which will add rule data to streams.
+ * When that happens, this struct holds all the UObject rule data for the stream while FConcertBaseStreamInfo will continue to hold non-UObject data.
+ */
+USTRUCT()
+struct FConcertReplicationStream
+{
+	GENERATED_BODY()
+
+	UPROPERTY()
+	FConcertBaseStreamInfo BaseDescription;
+
+	friend bool operator==(const FConcertReplicationStream& Left, const FConcertReplicationStream& Right)
+	{
+		return Left.BaseDescription == Right.BaseDescription;
+	}
+	friend bool operator!=(const FConcertReplicationStream& Left, const FConcertReplicationStream& Right)
+	{
+		return !(Left == Right);
+	}
+};
+
+inline FConcertBaseStreamInfo FConcertBaseStreamInfo::Clone(UE::ConcertSyncCore::EReplicationStreamCloneFlags Flags) const
+{
+	using namespace UE::ConcertSyncCore;
+	FConcertBaseStreamInfo Result { Identifier };
+
+	const TMap<FSoftObjectPath, FConcertReplicatedObjectInfo>& ReplicatedObjects = ReplicationMap.ReplicatedObjects;
+	if (EnumHasAnyFlags(Flags, EReplicationStreamCloneFlags::SkipProperties)
+		&& !ReplicationMap.IsEmpty())
+	{
+		Result.ReplicationMap.ReplicatedObjects.Reserve(ReplicatedObjects.Num());
+		for (const TPair<FSoftObjectPath, FConcertReplicatedObjectInfo>& ObjectInfo : ReplicatedObjects)
+		{
+			Result.ReplicationMap.ReplicatedObjects.Add(ObjectInfo.Key, { ObjectInfo.Value.ClassPath });
+		}
+	}
+	else
+	{
+		Result.ReplicationMap = ReplicationMap;
+	}
+
+	if (!EnumHasAnyFlags(Flags, EReplicationStreamCloneFlags::SkipFrequency))
+	{
+		Result.FrequencySettings = FrequencySettings;
+	}
+		
+	return Result;
+}
\ No newline at end of file
diff --git a/Engine/Plugins/Developer/Concert/ConcertSync/ConcertSyncCore/Source/ConcertSyncCore/Public/Replication/Data/ReplicationStreamDescription.h b/Engine/Plugins/Developer/Concert/ConcertSync/ConcertSyncCore/Source/ConcertSyncCore/Public/Replication/Data/ReplicationStreamDescription.h
deleted file mode 100644
diff --git a/Engine/Plugins/Developer/Concert/ConcertSync/ConcertSyncCore/Source/ConcertSyncCore/Public/Replication/Messages/ChangeStream.h b/Engine/Plugins/Developer/Concert/ConcertSync/ConcertSyncCore/Source/ConcertSyncCore/Public/Replication/Messages/ChangeStream.h
--- a/Engine/Plugins/Developer/Concert/ConcertSync/ConcertSyncCore/Source/ConcertSyncCore/Public/Replication/Messages/ChangeStream.h
+++ b/Engine/Plugins/Developer/Concert/ConcertSync/ConcertSyncCore/Source/ConcertSyncCore/Public/Replication/Messages/ChangeStream.h
@@ -5,14 +5,14 @@
 #include "EReplicationResponseErrorCode.h"
 #include "Replication/Data/ObjectIds.h"
 #include "Replication/Data/ConcertPropertySelection.h"
-#include "Replication/Data/ReplicationStreamDescription.h"
+#include "Replication/Data/ReplicationStream.h"
 
 #include "Misc/Optional.h"
 #include "ChangeStream.generated.h"
 
 class FOutputDevice;
 
-struct FReplicatedObjectInfo;
+struct FConcertReplicatedObjectInfo;
 
 /** A request to add a new object to a stream or overwrite a pre-existing object's properties / class. */
 USTRUCT()
@@ -51,15 +51,15 @@ struct FConcertReplication_ChangeStream_PutObject
 	}
 
 	// Intention: Ideally code dealing with PutObject requests uses these constructors / factory functions.
-	// If a property is added to FReplicatedObjectInfo, only the below code needs to be updated.
+	// If a property is added to FConcertReplicatedObjectInfo, only the below code needs to be updated.
 
 	/** @return The PutObject request if New contained sufficient info - empty otherwise */
-	CONCERTSYNCCORE_API static TOptional<FConcertReplication_ChangeStream_PutObject> MakeFromInfo(const FReplicatedObjectInfo& New);
+	CONCERTSYNCCORE_API static TOptional<FConcertReplication_ChangeStream_PutObject> MakeFromInfo(const FConcertReplicatedObjectInfo& New);
 	/** @return The PutObject request if changing Base to Desired contained sufficient info - empty otherwise */
-	CONCERTSYNCCORE_API static TOptional<FConcertReplication_ChangeStream_PutObject> MakeFromChange(const FReplicatedObjectInfo& Base, const FReplicatedObjectInfo& Desired);
+	CONCERTSYNCCORE_API static TOptional<FConcertReplication_ChangeStream_PutObject> MakeFromChange(const FConcertReplicatedObjectInfo& Base, const FConcertReplicatedObjectInfo& Desired);
 
 	/** Creates a new object info if there is sufficient data (all fields must be set for this). */
-	CONCERTSYNCCORE_API TOptional<FReplicatedObjectInfo> MakeObjectInfoIfValid() const;
+	CONCERTSYNCCORE_API TOptional<FConcertReplicatedObjectInfo> MakeObjectInfoIfValid() const;
 };
 
 UENUM()
@@ -132,7 +132,7 @@ struct FConcertReplication_ChangeStream_Request
 	 * If the requesting client has authority over these objects, authority is removed.
 	 */
 	UPROPERTY()
-	TSet<FObjectInStreamID> ObjectsToRemove;
+	TSet<FConcertObjectInStreamID> ObjectsToRemove;
 
 	/**
 	 * Adds new or modifies preexisting object definitions in pre-existing streams.
@@ -148,7 +148,7 @@ struct FConcertReplication_ChangeStream_Request
 	 * ambiguous requests, e.g. both StreamsToAdd and ObjectsToPut containing object Foo but with different property selections.
 	 */
 	UPROPERTY()
-	TMap<FObjectInStreamID, FConcertReplication_ChangeStream_PutObject> ObjectsToPut;
+	TMap<FConcertObjectInStreamID, FConcertReplication_ChangeStream_PutObject> ObjectsToPut;
 
 	/**
 	 * Changes the replication frequency settings for a stream.
@@ -172,7 +172,7 @@ struct FConcertReplication_ChangeStream_Request
 	 * in the stream's content being replaced.
 	 */
 	UPROPERTY()
-	TArray<FReplicationStreamDescription> StreamsToAdd;
+	TArray<FConcertReplicationStream> StreamsToAdd;
 
 	/**
 	 * Streams to remove from the server.
@@ -239,7 +239,7 @@ struct FConcertReplication_ChangeStream_FrequencyResponse
 
 	/** Streams that could not have their frequencies overriden. */
 	UPROPERTY()
-	TMap<FObjectInStreamID, EConcertChangeObjectFrequencyErrorCode> OverrideFailures;
+	TMap<FConcertObjectInStreamID, EConcertChangeObjectFrequencyErrorCode> OverrideFailures;
 
 	/** Streams that could not have their frequency defaults overriden. */
 	UPROPERTY()
@@ -278,11 +278,11 @@ struct FConcertReplication_ChangeStream_Response
 	 * Example: R has two streams S1 and S2. S1 contains the transform properties and S2 does not. Both S1 and S2 have authority over object Foo. It is legal to request S2 to contain the transform properties.
 	 */
 	UPROPERTY()
-	TMap<FObjectInStreamID, FReplicatedObjectId> AuthorityConflicts;
+	TMap<FConcertObjectInStreamID, FConcertReplicatedObjectId> AuthorityConflicts;
 
 	/** Reports semantic errors with ObjectsToPut. */
 	UPROPERTY()
-	TMap<FObjectInStreamID, EConcertPutObjectErrorCode> ObjectsToPutSemanticErrors;
+	TMap<FConcertObjectInStreamID, EConcertPutObjectErrorCode> ObjectsToPutSemanticErrors;
 	
 	/** Streams that were in StreamsToAdd but that were not created. */
 	UPROPERTY()
@@ -302,7 +302,7 @@ struct FConcertReplication_ChangeStream_Response
 	}
 	bool IsFailure() const { return !IsSuccess(); }
 
-	bool WasObjectPutSuccessful(const FObjectInStreamID& Object) const { return !AuthorityConflicts.Contains(Object) && !ObjectsToPutSemanticErrors.Contains(Object); }
+	bool WasObjectPutSuccessful(const FConcertObjectInStreamID& Object) const { return !AuthorityConflicts.Contains(Object) && !ObjectsToPutSemanticErrors.Contains(Object); }
 
 	/** If IsFailure(), logs the errors. */
 	CONCERTSYNCCORE_API void LogErrors(FOutputDevice& OutputDevice) const;
diff --git a/Engine/Plugins/Developer/Concert/ConcertSync/ConcertSyncCore/Source/ConcertSyncCore/Public/Replication/Messages/ClientQuery.h b/Engine/Plugins/Developer/Concert/ConcertSync/ConcertSyncCore/Source/ConcertSyncCore/Public/Replication/Messages/ClientQuery.h
--- a/Engine/Plugins/Developer/Concert/ConcertSync/ConcertSyncCore/Source/ConcertSyncCore/Public/Replication/Messages/ClientQuery.h
+++ b/Engine/Plugins/Developer/Concert/ConcertSync/ConcertSyncCore/Source/ConcertSyncCore/Public/Replication/Messages/ClientQuery.h
@@ -73,5 +73,5 @@ struct FConcertReplication_QueryReplicationInfo_Response
 	 * Value: Info about the client.
 	 */
 	UPROPERTY()
-	TMap<FGuid, FReplicationClientQueriedInfo> ClientInfo;
+	TMap<FGuid, FConcertQueriedClientInfo> ClientInfo;
 };
\ No newline at end of file
diff --git a/Engine/Plugins/Developer/Concert/ConcertSync/ConcertSyncCore/Source/ConcertSyncCore/Public/Replication/Messages/Handshake.h b/Engine/Plugins/Developer/Concert/ConcertSync/ConcertSyncCore/Source/ConcertSyncCore/Public/Replication/Messages/Handshake.h
--- a/Engine/Plugins/Developer/Concert/ConcertSync/ConcertSyncCore/Source/ConcertSyncCore/Public/Replication/Messages/Handshake.h
+++ b/Engine/Plugins/Developer/Concert/ConcertSync/ConcertSyncCore/Source/ConcertSyncCore/Public/Replication/Messages/Handshake.h
@@ -2,7 +2,7 @@
 
 #pragma once
 
-#include "Replication/Data/ReplicationStreamDescription.h"
+#include "Replication/Data/ReplicationStream.h"
 #include "Handshake.generated.h"
 
 class FText;
@@ -63,7 +63,7 @@ struct FConcertReplication_Join_Request
 
 	/** The data the client offers to send. */
 	UPROPERTY()
-	TArray<FReplicationStreamDescription> Streams;
+	TArray<FConcertReplicationStream> Streams;
 };
 
 /**
diff --git a/Engine/Plugins/Developer/Concert/ConcertSync/ConcertSyncCore/Source/ConcertSyncCore/Public/Replication/Processing/IReplicationDataSource.h b/Engine/Plugins/Developer/Concert/ConcertSync/ConcertSyncCore/Source/ConcertSyncCore/Public/Replication/Processing/IReplicationDataSource.h
--- a/Engine/Plugins/Developer/Concert/ConcertSync/ConcertSyncCore/Source/ConcertSyncCore/Public/Replication/Processing/IReplicationDataSource.h
+++ b/Engine/Plugins/Developer/Concert/ConcertSync/ConcertSyncCore/Source/ConcertSyncCore/Public/Replication/Processing/IReplicationDataSource.h
@@ -7,7 +7,7 @@
 #include "Misc/Optional.h"
 #include "Templates/Function.h"
 
-struct FReplicatedObjectId;
+struct FConcertReplicatedObjectId;
 
 namespace UE::ConcertSyncCore
 {
@@ -23,7 +23,7 @@ namespace UE::ConcertSyncCore
 	public:
 		
 		/** Iterates the objects that must be processed for replication (the result of ExtractReplicationDataForObject MAY return something new). */
-		virtual void ForEachPendingObject(TFunctionRef<void(const FReplicatedObjectId&)> ProcessItemFunc) const = 0;
+		virtual void ForEachPendingObject(TFunctionRef<void(const FConcertReplicatedObjectId&)> ProcessItemFunc) const = 0;
 		/** @return The number of objects ForEachPendingObject would iterate. Can be used, e.g. for Reserve()-ing a container. */
 		virtual int32 NumObjects() const = 0;
 
@@ -40,12 +40,12 @@ namespace UE::ConcertSyncCore
 		 * @return Whether successful. False indicates the call was invalid to make (ForEachPendingObject did not return Object). True indicates success, even if ProcessCopyable was not called.		
 		 */
 		virtual bool ExtractReplicationDataForObject(
-			const FReplicatedObjectId& Object,
+			const FConcertReplicatedObjectId& Object,
 			TFunctionRef<void(const FConcertSessionSerializedPayload& Payload)> ProcessCopyable,
 			TFunctionRef<void(FConcertSessionSerializedPayload&& Payload)> ProcessMoveable
 			) = 0;
 		/** Util version for callers that only want to read and do not want to store the payload. */
-		bool ExtractReplicationDataForObject(const FReplicatedObjectId& Object, TFunctionRef<void(const FConcertSessionSerializedPayload& Payload)> ProcessCopyable)
+		bool ExtractReplicationDataForObject(const FConcertReplicatedObjectId& Object, TFunctionRef<void(const FConcertSessionSerializedPayload& Payload)> ProcessCopyable)
 		{
 			return ExtractReplicationDataForObject(Object, ProcessCopyable, [&ProcessCopyable](FConcertSessionSerializedPayload&& Payload){ ProcessCopyable(Payload); });
 		}
diff --git a/Engine/Plugins/Developer/Concert/ConcertSync/ConcertSyncCore/Source/ConcertSyncCore/Public/Replication/Processing/ObjectReplicationCache.h b/Engine/Plugins/Developer/Concert/ConcertSync/ConcertSyncCore/Source/ConcertSyncCore/Public/Replication/Processing/ObjectReplicationCache.h
--- a/Engine/Plugins/Developer/Concert/ConcertSync/ConcertSyncCore/Source/ConcertSyncCore/Public/Replication/Processing/ObjectReplicationCache.h
+++ b/Engine/Plugins/Developer/Concert/ConcertSync/ConcertSyncCore/Source/ConcertSyncCore/Public/Replication/Processing/ObjectReplicationCache.h
@@ -25,7 +25,7 @@ namespace UE::ConcertSyncCore
 	public:
 
 		/** @return Whether this user is interested in data from this object. */
-		virtual bool WantsToAcceptObject(const FReplicatedObjectId& Object) const = 0;
+		virtual bool WantsToAcceptObject(const FConcertReplicatedObjectId& Object) const = 0;
 
 		/**
 		 * Called when data that is interesting to this user becomes available.
@@ -33,7 +33,7 @@ namespace UE::ConcertSyncCore
 		 * The user can keep hold of Data until it is used, at which point it just let's Data get out of scope.
 		 * If new data is received while this user is referencing Data, Data will be combined to contain any new data.
 		 */
-		virtual void OnDataCached(const FReplicatedObjectId& Object, TSharedRef<const FConcertReplication_ObjectReplicationEvent> Data) = 0;
+		virtual void OnDataCached(const FConcertReplicatedObjectId& Object, TSharedRef<const FConcertReplication_ObjectReplicationEvent> Data) = 0;
 
 		virtual ~IReplicationCacheUser() = default;
 	};
@@ -92,6 +92,6 @@ namespace UE::ConcertSyncCore
 			TMap<TWeakPtr<IReplicationCacheUser>, TWeakPtr<FConcertReplication_ObjectReplicationEvent>> DataInUse;
 		};
 		/** Maps every object to the events cached for it. */
-		TMap<FObjectInStreamID, FObjectCache> Cache; 
+		TMap<FConcertObjectInStreamID, FObjectCache> Cache; 
 	};
 }
\ No newline at end of file
diff --git a/Engine/Plugins/Developer/Concert/ConcertSync/ConcertSyncCore/Source/ConcertSyncCore/Public/Replication/Processing/ObjectReplicationProcessor.h b/Engine/Plugins/Developer/Concert/ConcertSync/ConcertSyncCore/Source/ConcertSyncCore/Public/Replication/Processing/ObjectReplicationProcessor.h
--- a/Engine/Plugins/Developer/Concert/ConcertSync/ConcertSyncCore/Source/ConcertSyncCore/Public/Replication/Processing/ObjectReplicationProcessor.h
+++ b/Engine/Plugins/Developer/Concert/ConcertSync/ConcertSyncCore/Source/ConcertSyncCore/Public/Replication/Processing/ObjectReplicationProcessor.h
@@ -9,7 +9,7 @@
 #include "UObject/SoftObjectPath.h"
 
 struct FConcertObjectReplicationSettings;
-struct FObjectInStreamID;
+struct FConcertObjectInStreamID;
 
 namespace UE::ConcertSyncCore
 {
@@ -68,7 +68,7 @@ namespace UE::ConcertSyncCore
 		struct FObjectProcessArgs
 		{
 			/** Info about the object to process */
-			FReplicatedObjectId ObjectInfo;
+			FConcertReplicatedObjectId ObjectInfo;
 		};
 
 		FORCEINLINE IReplicationDataSource& GetDataSource() const { return DataSource.Get(); }
diff --git a/Engine/Plugins/Developer/Concert/ConcertSync/ConcertSyncCore/Source/ConcertSyncCore/Public/Replication/Processing/Proxy/FrequencyTracker.h b/Engine/Plugins/Developer/Concert/ConcertSync/ConcertSyncCore/Source/ConcertSyncCore/Public/Replication/Processing/Proxy/FrequencyTracker.h
--- a/Engine/Plugins/Developer/Concert/ConcertSync/ConcertSyncCore/Source/ConcertSyncCore/Public/Replication/Processing/Proxy/FrequencyTracker.h
+++ b/Engine/Plugins/Developer/Concert/ConcertSync/ConcertSyncCore/Source/ConcertSyncCore/Public/Replication/Processing/Proxy/FrequencyTracker.h
@@ -42,7 +42,7 @@ namespace UE::ConcertSyncCore
 		 * @return Whether the object can be processed now
 		 */
 		ECanProcessObject TrackAndCheckObject(
-			const FObjectInStreamID& ObjectPath,
+			const FConcertObjectInStreamID& ObjectPath,
 			const FConcertObjectReplicationSettings& ReplicationSettings,
 			const FPlatformSecondsTimestamp NowAsPlatformSeconds = FPlatformTime::Seconds()
 			);
@@ -54,7 +54,7 @@ namespace UE::ConcertSyncCore
 		 * @param Object The object that was processed
 		 * @param NowAsPlatformSeconds The current time - result of FPlatformTime::Seconds(). Pass this in to avoid excessive time spent on timing.
 		 */
-		void OnProcessObject(const FObjectInStreamID& Object, const FPlatformSecondsTimestamp NowAsPlatformSeconds = FPlatformTime::Seconds());
+		void OnProcessObject(const FConcertObjectInStreamID& Object, const FPlatformSecondsTimestamp NowAsPlatformSeconds = FPlatformTime::Seconds());
 		
 		/**
 		 * Removes all tracked objects were last processed before CutoffTime.
@@ -65,12 +65,12 @@ namespace UE::ConcertSyncCore
 	protected:
 
 		// For unit tests
-		const TMap<FObjectInStreamID, FPlatformSecondsTimestamp>& GetObjectsToLastProcessed() const { return ObjectsToLastProcessed; }
+		const TMap<FConcertObjectInStreamID, FPlatformSecondsTimestamp>& GetObjectsToLastProcessed() const { return ObjectsToLastProcessed; }
 		
 	private:
 
 		/** It is important to call CleanseOutdatedObjects to avoid this mapping leaking objects that used to be replicated. */
-		TMap<FObjectInStreamID, FPlatformSecondsTimestamp> ObjectsToLastProcessed;
+		TMap<FConcertObjectInStreamID, FPlatformSecondsTimestamp> ObjectsToLastProcessed;
 	};
 
 	/** Subclass of FFrequencyTracker that you call Tick on. Every x seconds, this will call CleanseOutdatedObjects itself. */
diff --git a/Engine/Plugins/Developer/Concert/ConcertSync/ConcertSyncCore/Source/ConcertSyncCore/Public/Replication/Processing/Proxy/ObjectProcessorProxy_Frequency.h b/Engine/Plugins/Developer/Concert/ConcertSync/ConcertSyncCore/Source/ConcertSyncCore/Public/Replication/Processing/Proxy/ObjectProcessorProxy_Frequency.h
--- a/Engine/Plugins/Developer/Concert/ConcertSync/ConcertSyncCore/Source/ConcertSyncCore/Public/Replication/Processing/Proxy/ObjectProcessorProxy_Frequency.h
+++ b/Engine/Plugins/Developer/Concert/ConcertSync/ConcertSyncCore/Source/ConcertSyncCore/Public/Replication/Processing/Proxy/ObjectProcessorProxy_Frequency.h
@@ -11,7 +11,7 @@
 
 namespace UE::ConcertSyncCore
 {
-	DECLARE_DELEGATE_RetVal_OneParam(FConcertObjectReplicationSettings, FGetObjectFrequencySettings, const FReplicatedObjectId& Object);
+	DECLARE_DELEGATE_RetVal_OneParam(FConcertObjectReplicationSettings, FGetObjectFrequencySettings, const FConcertReplicatedObjectId& Object);
 	
 	/**
 	 * Maps every FSoftObjectPath to FConcertObjectReplicationSettings.
@@ -83,7 +83,7 @@ namespace UE::ConcertSyncCore
 	template <TIsObjectProcessorConcept TRealProcessorImpl>
 	void TObjectProcessorProxy_Frequency<TRealProcessorImpl>::ProcessObject(const FObjectReplicationProcessor::FObjectProcessArgs& Args)
 	{
-		const FReplicatedObjectId& Object = Args.ObjectInfo;
+		const FConcertReplicatedObjectId& Object = Args.ObjectInfo;
 		const FConcertObjectReplicationSettings FrequencySettings = GetObjectFrequencySettingsDelegate.Execute(Object);
 
 		// Throttle calls to the real ProcessObject implementation based on the object's frequency settings
diff --git a/Engine/Plugins/Developer/Concert/ConcertSync/ConcertSyncCore/Source/ConcertSyncCore/Public/Replication/Processing/ReplicationDataQueuer.h b/Engine/Plugins/Developer/Concert/ConcertSync/ConcertSyncCore/Source/ConcertSyncCore/Public/Replication/Processing/ReplicationDataQueuer.h
--- a/Engine/Plugins/Developer/Concert/ConcertSync/ConcertSyncCore/Source/ConcertSyncCore/Public/Replication/Processing/ReplicationDataQueuer.h
+++ b/Engine/Plugins/Developer/Concert/ConcertSync/ConcertSyncCore/Source/ConcertSyncCore/Public/Replication/Processing/ReplicationDataQueuer.h
@@ -27,13 +27,13 @@ namespace UE::ConcertSyncCore
 		virtual ~FReplicationDataQueuer() override;
 
 		//~ Begin IReplicationDataSource Interface
-		virtual void ForEachPendingObject(TFunctionRef<void(const FReplicatedObjectId&)> ProcessItemFunc) const override;
+		virtual void ForEachPendingObject(TFunctionRef<void(const FConcertReplicatedObjectId&)> ProcessItemFunc) const override;
 		virtual int32 NumObjects() const override;
-		virtual bool ExtractReplicationDataForObject(const FReplicatedObjectId& Object, TFunctionRef<void(const FConcertSessionSerializedPayload& Payload)> ProcessCopyable, TFunctionRef<void(FConcertSessionSerializedPayload&& Payload)> ProcessMoveable) override;
+		virtual bool ExtractReplicationDataForObject(const FConcertReplicatedObjectId& Object, TFunctionRef<void(const FConcertSessionSerializedPayload& Payload)> ProcessCopyable, TFunctionRef<void(FConcertSessionSerializedPayload&& Payload)> ProcessMoveable) override;
 		//~ End IReplicationDataSource Interface
 		
 		//~ Begin IReplicationCacheUser Interface
-		virtual void OnDataCached(const FReplicatedObjectId& Object, TSharedRef<const FConcertReplication_ObjectReplicationEvent> Data) override;
+		virtual void OnDataCached(const FConcertReplicatedObjectId& Object, TSharedRef<const FConcertReplication_ObjectReplicationEvent> Data) override;
 		//~ End IReplicationCacheUser Interface
 
 	protected:
@@ -44,7 +44,7 @@ namespace UE::ConcertSyncCore
 	private:
 
 		/** Stores events as they are received. */
-		TMap<FReplicatedObjectId, TSharedPtr<const FConcertReplication_ObjectReplicationEvent>> PendingEvents;
+		TMap<FConcertReplicatedObjectId, TSharedPtr<const FConcertReplication_ObjectReplicationEvent>> PendingEvents;
 
 		/** Provides us with replication events and shares them effectively. */
 		TSharedPtr<FObjectReplicationCache> ReplicationCache;
diff --git a/Engine/Plugins/Developer/Concert/ConcertSync/ConcertSyncServer/Source/ConcertSyncServer/Private/ConcertSyncServer.cpp b/Engine/Plugins/Developer/Concert/ConcertSync/ConcertSyncServer/Source/ConcertSyncServer/Private/ConcertSyncServer.cpp
--- a/Engine/Plugins/Developer/Concert/ConcertSync/ConcertSyncServer/Source/ConcertSyncServer/Private/ConcertSyncServer.cpp
+++ b/Engine/Plugins/Developer/Concert/ConcertSync/ConcertSyncServer/Source/ConcertSyncServer/Private/ConcertSyncServer.cpp
@@ -619,10 +619,10 @@ void FConcertSyncServer::LogReplicationStreams() const
 		{
 			const FGuid ClientId = ClientInfo.ClientEndpointId;
 			StringBuilder << ClientIndent << TEXT("Client: ") << *ClientInfo.ClientInfo.DisplayName << TEXT(" (") << ClientId.ToString() << TEXT(")");
-			ReplicationManagerPair.Value->ForEachStream(ClientId, [&StringBuilder, StreamIndent, ObjectIndent, PropertyIndent](const FReplicationStreamDescription& Stream)
+			ReplicationManagerPair.Value->ForEachStream(ClientId, [&StringBuilder, StreamIndent, ObjectIndent, PropertyIndent](const FConcertReplicationStream& Stream)
 			{
 				StringBuilder << StreamIndent << TEXT("Stream: ") << *Stream.BaseDescription.Identifier.ToString();
-				for (const TPair<FSoftObjectPath, FReplicatedObjectInfo>& ReplicatedObject : Stream.BaseDescription.ReplicationMap.ReplicatedObjects)
+				for (const TPair<FSoftObjectPath, FConcertReplicatedObjectInfo>& ReplicatedObject : Stream.BaseDescription.ReplicationMap.ReplicatedObjects)
 				{
 					StringBuilder << ObjectIndent << *ReplicatedObject.Key.ToString();
 					for (const FConcertPropertyChain& Property : ReplicatedObject.Value.PropertySelection.ReplicatedProperties)
@@ -663,7 +663,7 @@ void FConcertSyncServer::LogReplicationAuthority() const
 			StringBuilder << ClientIndent << TEXT("Client: ") << *ClientInfo.ClientInfo.DisplayName << TEXT(" (") << ClientId.ToString() << TEXT(")");
 			
 			const TSharedRef<FConcertServerReplicationManager>& ReplicationManager = ReplicationManagerPair.Value;
-			ReplicationManager->ForEachStream(ClientId, [&StringBuilder, StreamIndent, ObjectIndent, ClientId, &ReplicationManager](const FReplicationStreamDescription& Stream)
+			ReplicationManager->ForEachStream(ClientId, [&StringBuilder, StreamIndent, ObjectIndent, ClientId, &ReplicationManager](const FConcertReplicationStream& Stream)
 			{
 				const FGuid StreamId = Stream.BaseDescription.Identifier;
 				StringBuilder << StreamIndent << TEXT("Stream: ") << *StreamId.ToString();
diff --git a/Engine/Plugins/Developer/Concert/ConcertSync/ConcertSyncServer/Source/ConcertSyncServer/Private/Replication/AuthorityManager.cpp b/Engine/Plugins/Developer/Concert/ConcertSync/ConcertSyncServer/Source/ConcertSyncServer/Private/Replication/AuthorityManager.cpp
--- a/Engine/Plugins/Developer/Concert/ConcertSync/ConcertSyncServer/Source/ConcertSyncServer/Private/Replication/AuthorityManager.cpp
+++ b/Engine/Plugins/Developer/Concert/ConcertSync/ConcertSyncServer/Source/ConcertSyncServer/Private/Replication/AuthorityManager.cpp
@@ -6,7 +6,7 @@
 #include "IConcertSession.h"
 #include "Replication/AuthorityConflictSharedUtils.h"
 #include "Replication/Data/ObjectIds.h"
-#include "Replication/Data/ReplicationStreamDescription.h"
+#include "Replication/Data/ReplicationStream.h"
 #include "Replication/Messages/ChangeAuthority.h"
 
 namespace UE::ConcertSyncServer::Replication
@@ -24,9 +24,9 @@ namespace UE::ConcertSyncServer::Replication
 				, Getters(Getters)
 			{}
 
-			virtual void ForEachStream(const FGuid& ClientEndpointId, TFunctionRef<EBreakBehavior(const FGuid& StreamId, const FObjectReplicationMap& ReplicationMap)> Callback) const override
+			virtual void ForEachStream(const FGuid& ClientEndpointId, TFunctionRef<EBreakBehavior(const FGuid& StreamId, const FConcertObjectReplicationMap& ReplicationMap)> Callback) const override
 			{
-				Getters.ForEachStream(ClientEndpointId, [&Callback](const FReplicationStreamDescription& Stream)
+				Getters.ForEachStream(ClientEndpointId, [&Callback](const FConcertReplicationStream& Stream)
 				{
 					return Callback(Stream.BaseDescription.Identifier, Stream.BaseDescription.ReplicationMap);
 				});
@@ -70,7 +70,7 @@ namespace UE::ConcertSyncServer::Replication
 		Session->UnregisterCustomRequestHandler<FConcertReplication_ChangeAuthority_Request>();
 	}
 
-	bool FAuthorityManager::HasAuthorityToChange(const FReplicatedObjectId& ObjectChange) const
+	bool FAuthorityManager::HasAuthorityToChange(const FConcertReplicatedObjectId& ObjectChange) const
 	{
 		const FClientAuthorityData* AuthorityData = ClientAuthorityData.Find(ObjectChange.SenderEndpointId);
 		const TSet<FSoftObjectPath>* OwnedObjects = AuthorityData ? AuthorityData->OwnedObjects.Find(ObjectChange.StreamId) : nullptr;
@@ -96,7 +96,7 @@ namespace UE::ConcertSyncServer::Replication
 	}
 
 	FAuthorityManager::EAuthorityResult FAuthorityManager::EnumerateAuthorityConflicts(
-		const FReplicatedObjectId& Object, 
+		const FConcertReplicatedObjectId& Object, 
 		const FConcertPropertySelection* OverwriteProperties,
 		FProcessAuthorityConflict ProcessConflict
 		) const
@@ -105,8 +105,8 @@ namespace UE::ConcertSyncServer::Replication
 		const FConcertPropertySelection* PropertiesToCheck = OverwriteProperties;
 		if (!PropertiesToCheck)
 		{
-			const FReplicationStreamDescription* Description = FindClientStreamById(ClientId, Object.StreamId);
-			const FReplicatedObjectInfo* PropertyInfo = Description ? Description->BaseDescription.ReplicationMap.ReplicatedObjects.Find(Object.Object) : nullptr;
+			const FConcertReplicationStream* Description = FindClientStreamById(ClientId, Object.StreamId);
+			const FConcertReplicatedObjectInfo* PropertyInfo = Description ? Description->BaseDescription.ReplicationMap.ReplicatedObjects.Find(Object.Object) : nullptr;
 			PropertiesToCheck = PropertyInfo ? &PropertyInfo->PropertySelection : nullptr;
 		}
 
@@ -130,7 +130,7 @@ namespace UE::ConcertSyncServer::Replication
 		return Conflict == EAuthorityConflict::Allowed ? EAuthorityResult::Allowed : EAuthorityResult::Conflict;
 	}
 
-	bool FAuthorityManager::CanTakeAuthority(const FReplicatedObjectId& Object) const
+	bool FAuthorityManager::CanTakeAuthority(const FConcertReplicatedObjectId& Object) const
 	{
 		return EnumerateAuthorityConflicts(Object) == EAuthorityResult::Allowed;
 	}
@@ -140,7 +140,7 @@ namespace UE::ConcertSyncServer::Replication
 		ClientAuthorityData.Remove(ClientEndpointId);
 	}
 
-	void FAuthorityManager::RemoveAuthority(const FReplicatedObjectId& Object)
+	void FAuthorityManager::RemoveAuthority(const FConcertReplicatedObjectId& Object)
 	{
 		const FClientId& ClientId = Object.SenderEndpointId;
 		FClientAuthorityData* ClientData = ClientAuthorityData.Find(ClientId);
@@ -183,7 +183,7 @@ namespace UE::ConcertSyncServer::Replication
 		const FClientId& ClientId = ConcertSessionContext.SourceEndpointId;
 		Private::ForEachReplicatedObject(Request.TakeAuthority, [this, &Response, &AuthorityData, &ClientId](const FStreamId& StreamId, const FSoftObjectPath& ObjectPath)
 		{
-			const FReplicatedObjectId ObjectToAuthor{ { StreamId, ObjectPath }, ClientId };
+			const FConcertReplicatedObjectId ObjectToAuthor{ { StreamId, ObjectPath }, ClientId };
 			if (CanTakeAuthority(ObjectToAuthor))
 			{
 				UE_LOG(LogConcert, Log, TEXT("Transferred authority of %s to client %s for their stream %s"), *ObjectPath.ToString(), *ClientId.ToString(EGuidFormats::Short), *StreamId.ToString(EGuidFormats::Short));
@@ -217,10 +217,10 @@ namespace UE::ConcertSyncServer::Replication
 		return EConcertSessionResponseCode::Success;
 	}
 	
-	const FReplicationStreamDescription* FAuthorityManager::FindClientStreamById(const FClientId& ClientId, const FStreamId& StreamId) const
+	const FConcertReplicationStream* FAuthorityManager::FindClientStreamById(const FClientId& ClientId, const FStreamId& StreamId) const
 	{
-		const FReplicationStreamDescription* StreamDescription = nullptr;
-		Getters.ForEachStream(ClientId, [&StreamId, &StreamDescription](const FReplicationStreamDescription& Stream) mutable
+		const FConcertReplicationStream* StreamDescription = nullptr;
+		Getters.ForEachStream(ClientId, [&StreamId, &StreamDescription](const FConcertReplicationStream& Stream) mutable
 		{
 			if (Stream.BaseDescription.Identifier == StreamId)
 			{
diff --git a/Engine/Plugins/Developer/Concert/ConcertSync/ConcertSyncServer/Source/ConcertSyncServer/Private/Replication/AuthorityManager.h b/Engine/Plugins/Developer/Concert/ConcertSync/ConcertSyncServer/Source/ConcertSyncServer/Private/Replication/AuthorityManager.h
--- a/Engine/Plugins/Developer/Concert/ConcertSync/ConcertSyncServer/Source/ConcertSyncServer/Private/Replication/AuthorityManager.h
+++ b/Engine/Plugins/Developer/Concert/ConcertSync/ConcertSyncServer/Source/ConcertSyncServer/Private/Replication/AuthorityManager.h
@@ -14,9 +14,9 @@ struct FConcertReplication_ChangeAuthority_Request;
 struct FConcertPropertyChain;
 struct FConcertPropertySelection;
 struct FConcertSessionContext;
-struct FObjectInStreamID;
-struct FReplicatedObjectId;
-struct FReplicationStreamDescription;
+struct FConcertObjectInStreamID;
+struct FConcertReplicatedObjectId;
+struct FConcertReplicationStream;
 
 namespace UE::ConcertSyncServer::Replication
 {
@@ -29,7 +29,7 @@ namespace UE::ConcertSyncServer::Replication
 	public:
 
 		/** Provides a way to extract all streams registered to a given client. */
-		virtual void ForEachStream(const FGuid& ClientEndpointId, TFunctionRef<EBreakBehavior(const FReplicationStreamDescription& Stream)> Callback) const = 0;
+		virtual void ForEachStream(const FGuid& ClientEndpointId, TFunctionRef<EBreakBehavior(const FConcertReplicationStream& Stream)> Callback) const = 0;
 
 		/** Iterates through all clients have registered to send any data. */
 		virtual void ForEachSendingClient(TFunctionRef<EBreakBehavior(const FGuid& ClientEndpointId)> Callback) const = 0;
@@ -53,7 +53,7 @@ namespace UE::ConcertSyncServer::Replication
 		 * Checks whether the client that sent the identified object had authority to send it.
 		 * @return Whether the server should process the object change.
 		 */
-		bool HasAuthorityToChange(const FReplicatedObjectId& ObjectChange) const;
+		bool HasAuthorityToChange(const FConcertReplicatedObjectId& ObjectChange) const;
 
 		/** Utility for iterating authority a client has for a given stream. */
 		void EnumerateAuthority(const FClientId& ClientId, const FStreamId& StreamId, TFunctionRef<EBreakBehavior(const FSoftObjectPath& Object)> Callback) const;
@@ -73,17 +73,17 @@ namespace UE::ConcertSyncServer::Replication
 		 * @return Whether there were any conflicts
 		 */
 		EAuthorityResult EnumerateAuthorityConflicts(
-			const FReplicatedObjectId& Object,
+			const FConcertReplicatedObjectId& Object,
 			const FConcertPropertySelection* OverwriteProperties = nullptr,
 			FProcessAuthorityConflict ProcessConflict = [](auto&, auto&, auto&){ return EBreakBehavior::Break; }
 			) const;
 		/** Whether it is legal for this the client identified by ClientId to take control over the object given the stream the client has registered. */
-		bool CanTakeAuthority(const FReplicatedObjectId& Object) const;
+		bool CanTakeAuthority(const FConcertReplicatedObjectId& Object) const;
 
 		/** Notifies this manager that the client has left, which means all their authority is now gone. */
 		void OnClientLeft(const FClientId& ClientEndpointId);
 		/** Takes away authority from the given client from the given object. */
-		void RemoveAuthority(const FReplicatedObjectId& Object);
+		void RemoveAuthority(const FConcertReplicatedObjectId& Object);
 
 	private:
 
@@ -106,7 +106,7 @@ namespace UE::ConcertSyncServer::Replication
 			);
 
 		/** Finds a stream registered with the client by its ID. */
-		const FReplicationStreamDescription* FindClientStreamById(const FClientId& ClientId, const FStreamId& StreamId) const;
+		const FConcertReplicationStream* FindClientStreamById(const FClientId& ClientId, const FStreamId& StreamId) const;
 	};
 }
 
diff --git a/Engine/Plugins/Developer/Concert/ConcertSync/ConcertSyncServer/Source/ConcertSyncServer/Private/Replication/ConcertReplicationClient.cpp b/Engine/Plugins/Developer/Concert/ConcertSync/ConcertSyncServer/Source/ConcertSyncServer/Private/Replication/ConcertReplicationClient.cpp
--- a/Engine/Plugins/Developer/Concert/ConcertSync/ConcertSyncServer/Source/ConcertSyncServer/Private/Replication/ConcertReplicationClient.cpp
+++ b/Engine/Plugins/Developer/Concert/ConcertSync/ConcertSyncServer/Source/ConcertSyncServer/Private/Replication/ConcertReplicationClient.cpp
@@ -9,7 +9,7 @@
 namespace UE::ConcertSyncServer::Replication
 {
 	FConcertReplicationClient::FConcertReplicationClient(
-		TArray<FReplicationStreamDescription> StreamDescriptions,
+		TArray<FConcertReplicationStream> StreamDescriptions,
 		const FGuid& ClientEndpointId,
 		TSharedRef<IConcertSession> Session,
 		TSharedRef<ConcertSyncCore::FObjectReplicationCache> ReplicationCache,
diff --git a/Engine/Plugins/Developer/Concert/ConcertSync/ConcertSyncServer/Source/ConcertSyncServer/Private/Replication/ConcertReplicationClient.h b/Engine/Plugins/Developer/Concert/ConcertSync/ConcertSyncServer/Source/ConcertSyncServer/Private/Replication/ConcertReplicationClient.h
--- a/Engine/Plugins/Developer/Concert/ConcertSync/ConcertSyncServer/Source/ConcertSyncServer/Private/Replication/ConcertReplicationClient.h
+++ b/Engine/Plugins/Developer/Concert/ConcertSync/ConcertSyncServer/Source/ConcertSyncServer/Private/Replication/ConcertReplicationClient.h
@@ -7,7 +7,7 @@
 #include "Replication/Processing/Proxy/ObjectProcessorProxy_Frequency.h"
 #include "Templates/SharedPointer.h"
 
-struct FReplicationStreamDescription;
+struct FConcertReplicationStream;
 struct FConcertReplication_Join_Request;
 struct FConcertReplication_ChangeStream_Request;
 
@@ -27,7 +27,7 @@ namespace UE::ConcertSyncServer::Replication
 	public:
 
 		FConcertReplicationClient(
-			TArray<FReplicationStreamDescription> StreamDescriptions,
+			TArray<FConcertReplicationStream> StreamDescriptions,
 			const FGuid& ClientEndpointId,
 			TSharedRef<IConcertSession> Session,
 			TSharedRef<ConcertSyncCore::FObjectReplicationCache> ReplicationCache,
@@ -44,12 +44,12 @@ namespace UE::ConcertSyncServer::Replication
 		void ApplyValidatedRequest(const FConcertReplication_ChangeStream_Request& Request);
 
 		const FGuid& GetClientEndpointId() const { return ClientEndpointId; }
-		const TArray<FReplicationStreamDescription>& GetStreamDescriptions() const { return StreamDescriptions; }
+		const TArray<FConcertReplicationStream>& GetStreamDescriptions() const { return StreamDescriptions; }
 
 	private:
 
 		/** The streams this client offered to send. */
-		TArray<FReplicationStreamDescription> StreamDescriptions;
+		TArray<FConcertReplicationStream> StreamDescriptions;
 		
 		/** This client's endpoint ID. */
 		const FGuid ClientEndpointId;
diff --git a/Engine/Plugins/Developer/Concert/ConcertSync/ConcertSyncServer/Source/ConcertSyncServer/Private/Replication/ConcertServerReplicationManager.cpp b/Engine/Plugins/Developer/Concert/ConcertSync/ConcertSyncServer/Source/ConcertSyncServer/Private/Replication/ConcertServerReplicationManager.cpp
--- a/Engine/Plugins/Developer/Concert/ConcertSync/ConcertSyncServer/Source/ConcertSyncServer/Private/Replication/ConcertServerReplicationManager.cpp
+++ b/Engine/Plugins/Developer/Concert/ConcertSync/ConcertSyncServer/Source/ConcertSyncServer/Private/Replication/ConcertServerReplicationManager.cpp
@@ -41,7 +41,7 @@ namespace UE::ConcertSyncServer::Replication
 		Session->OnTick().RemoveAll(this);
 	}
 
-	void FConcertServerReplicationManager::ForEachStream(const FGuid& ClientEndpointId, TFunctionRef<EBreakBehavior(const FReplicationStreamDescription& Stream)> Callback) const
+	void FConcertServerReplicationManager::ForEachStream(const FGuid& ClientEndpointId, TFunctionRef<EBreakBehavior(const FConcertReplicationStream& Stream)> Callback) const
 	{
 		const TUniquePtr<FConcertReplicationClient>* Client = Clients.Find(ClientEndpointId);
 		if (!ensure(Client))
@@ -49,7 +49,7 @@ namespace UE::ConcertSyncServer::Replication
 			return;
 		}
 
-		for (const FReplicationStreamDescription& Stream : (*Client)->GetStreamDescriptions())
+		for (const FConcertReplicationStream& Stream : (*Client)->GetStreamDescriptions())
 		{
 			if (Callback(Stream) == EBreakBehavior::Break)
 			{
@@ -142,7 +142,7 @@ namespace UE::ConcertSyncServer::Replication
 				continue;
 			}
 			
-			FReplicationClientQueriedInfo& EndpointInfo = Response.ClientInfo.Add(EndpointId);
+			FConcertQueriedClientInfo& EndpointInfo = Response.ClientInfo.Add(EndpointId);
 			if (!EnumHasAnyFlags(Request.QueryFlags, EConcertQueryClientStreamFlags::SkipStreamInfo))
 			{
 				EndpointInfo.Streams = BuildClientStreamInfo(*Client->Get(), Request.QueryFlags);
@@ -158,10 +158,10 @@ namespace UE::ConcertSyncServer::Replication
 		return EConcertSessionResponseCode::Success;
 	}
 
-	TArray<FSharedReplicationStreamDescription> FConcertServerReplicationManager::BuildClientStreamInfo(const FConcertReplicationClient& Client, EConcertQueryClientStreamFlags QueryFlags)
+	TArray<FConcertBaseStreamInfo> FConcertServerReplicationManager::BuildClientStreamInfo(const FConcertReplicationClient& Client, EConcertQueryClientStreamFlags QueryFlags)
 	{
-		TArray<FSharedReplicationStreamDescription> Result;
-		Algo::Transform(Client.GetStreamDescriptions(), Result, [QueryFlags](const FReplicationStreamDescription& Description)
+		TArray<FConcertBaseStreamInfo> Result;
+		Algo::Transform(Client.GetStreamDescriptions(), Result, [QueryFlags](const FConcertReplicationStream& Description)
 		{
 			using namespace ConcertSyncCore;
 			
@@ -180,18 +180,18 @@ namespace UE::ConcertSyncServer::Replication
 		return Result;
 	}
 
-	TArray<FReplicationAuthorityInfo> FConcertServerReplicationManager::BuildClientAuthorityInfo(const FConcertReplicationClient& Client) const
+	TArray<FConcertAuthorityClientInfo> FConcertServerReplicationManager::BuildClientAuthorityInfo(const FConcertReplicationClient& Client) const
 	{
-		TArray<FReplicationAuthorityInfo> Result;
-		for (const FReplicationStreamDescription& Description : Client.GetStreamDescriptions())
+		TArray<FConcertAuthorityClientInfo> Result;
+		for (const FConcertReplicationStream& Description : Client.GetStreamDescriptions())
 		{
 			const FGuid StreamId = Description.BaseDescription.Identifier;
-			FReplicationAuthorityInfo Info;
+			FConcertAuthorityClientInfo Info;
 			Info.StreamId = StreamId;
 			
-			for (const TPair<FSoftObjectPath, FReplicatedObjectInfo>& Pair : Description.BaseDescription.ReplicationMap.ReplicatedObjects)
+			for (const TPair<FSoftObjectPath, FConcertReplicatedObjectInfo>& Pair : Description.BaseDescription.ReplicationMap.ReplicatedObjects)
 			{
-				FReplicatedObjectId ObjectInfo;
+				FConcertReplicatedObjectId ObjectInfo;
 				ObjectInfo.SenderEndpointId = Client.GetClientEndpointId();
 				ObjectInfo.Object = Pair.Key;
 				ObjectInfo.StreamId = StreamId;
@@ -242,7 +242,7 @@ namespace UE::ConcertSyncServer::Replication
 		}
 	}
 
-	FConcertObjectReplicationSettings FConcertServerReplicationManager::GetObjectFrequencySettings(const FReplicatedObjectId& Object) const
+	FConcertObjectReplicationSettings FConcertServerReplicationManager::GetObjectFrequencySettings(const FConcertReplicatedObjectId& Object) const
 	{
 		const TUniquePtr<FConcertReplicationClient>* Client = Clients.Find(Object.SenderEndpointId);
 		if (!ensureMsgf(Client, TEXT("Caller is trying to retrieve non-existing client")))
@@ -251,7 +251,7 @@ namespace UE::ConcertSyncServer::Replication
 			return {};
 		}
 		
-		const FReplicationStreamDescription* Stream = Client->Get()->GetStreamDescriptions().FindByPredicate([&Object](const FReplicationStreamDescription& Description)
+		const FConcertReplicationStream* Stream = Client->Get()->GetStreamDescriptions().FindByPredicate([&Object](const FConcertReplicationStream& Description)
 			{
 				return Description.BaseDescription.Identifier == Object.StreamId;
 			});
diff --git a/Engine/Plugins/Developer/Concert/ConcertSync/ConcertSyncServer/Source/ConcertSyncServer/Private/Replication/ConcertServerReplicationManager.h b/Engine/Plugins/Developer/Concert/ConcertSync/ConcertSyncServer/Source/ConcertSyncServer/Private/Replication/ConcertServerReplicationManager.h
--- a/Engine/Plugins/Developer/Concert/ConcertSync/ConcertSyncServer/Source/ConcertSyncServer/Private/Replication/ConcertServerReplicationManager.h
+++ b/Engine/Plugins/Developer/Concert/ConcertSync/ConcertSyncServer/Source/ConcertSyncServer/Private/Replication/ConcertServerReplicationManager.h
@@ -21,7 +21,7 @@ enum class EConcertQueryClientStreamFlags : uint8;
 struct FConcertReplication_ChangeStream_Response;
 struct FConcertReplication_QueryReplicationInfo_Response;
 struct FConcertReplication_QueryReplicationInfo_Request;
-struct FReplicationAuthorityInfo;
+struct FConcertAuthorityClientInfo;
 
 namespace UE::ConcertSyncCore
 {
@@ -52,7 +52,7 @@ namespace UE::ConcertSyncServer::Replication
 		const FAuthorityManager& GetAuthorityManager() const { return AuthorityManager.Get(); }
 
 		//~ Begin IAuthorityManagerGetters Interface
-		virtual void ForEachStream(const FGuid& ClientEndpointId, TFunctionRef<EBreakBehavior(const FReplicationStreamDescription& Stream)> Callback) const override;
+		virtual void ForEachStream(const FGuid& ClientEndpointId, TFunctionRef<EBreakBehavior(const FConcertReplicationStream& Stream)> Callback) const override;
 		virtual void ForEachSendingClient(TFunctionRef<EBreakBehavior(const FGuid& ClientEndpointId)> Callback) const override;
 		//~ End IAuthorityManagerGetters Interface
 
@@ -85,9 +85,9 @@ namespace UE::ConcertSyncServer::Replication
 		// Querying
 		EConcertSessionResponseCode HandleQueryReplicationInfoRequest(const FConcertSessionContext& ConcertSessionContext, const FConcertReplication_QueryReplicationInfo_Request& Request, FConcertReplication_QueryReplicationInfo_Response& Response);
 		/** Gets all registered streams and optionally removes the properties. */
-		static TArray<FSharedReplicationStreamDescription> BuildClientStreamInfo(const FConcertReplicationClient& Client, EConcertQueryClientStreamFlags QueryFlags);
+		static TArray<FConcertBaseStreamInfo> BuildClientStreamInfo(const FConcertReplicationClient& Client, EConcertQueryClientStreamFlags QueryFlags);
 		/** Maps the client's streams to the objects in that stream the client has taken authority over. */
-		TArray<FReplicationAuthorityInfo> BuildClientAuthorityInfo(const FConcertReplicationClient& Client) const;
+		TArray<FConcertAuthorityClientInfo> BuildClientAuthorityInfo(const FConcertReplicationClient& Client) const;
 
 		// Changing streams
 		EConcertSessionResponseCode HandleChangeStreamRequest(const FConcertSessionContext& ConcertSessionContext, const FConcertReplication_ChangeStream_Request& Request, FConcertReplication_ChangeStream_Response& Response);
@@ -105,6 +105,6 @@ namespace UE::ConcertSyncServer::Replication
 		void Tick(IConcertServerSession& InSession, float InDeltaTime);
 		
 		/** Callback to clients for obtaining an object's frequency settings. */
-		FConcertObjectReplicationSettings GetObjectFrequencySettings(const FReplicatedObjectId& Object) const;
+		FConcertObjectReplicationSettings GetObjectFrequencySettings(const FConcertReplicatedObjectId& Object) const;
 	};
 }
\ No newline at end of file
diff --git a/Engine/Plugins/Developer/Concert/ConcertSync/ConcertSyncServer/Source/ConcertSyncServer/Private/Replication/ConcertServerReplicationManager_ChangeStream.cpp b/Engine/Plugins/Developer/Concert/ConcertSync/ConcertSyncServer/Source/ConcertSyncServer/Private/Replication/ConcertServerReplicationManager_ChangeStream.cpp
--- a/Engine/Plugins/Developer/Concert/ConcertSync/ConcertSyncServer/Source/ConcertSyncServer/Private/Replication/ConcertServerReplicationManager_ChangeStream.cpp
+++ b/Engine/Plugins/Developer/Concert/ConcertSync/ConcertSyncServer/Source/ConcertSyncServer/Private/Replication/ConcertServerReplicationManager_ChangeStream.cpp
@@ -14,9 +14,9 @@ namespace UE::ConcertSyncServer::Replication
 {
 	namespace Private
 	{
-		static const FReplicationStreamDescription* FindExistingStream(const FConcertReplicationClient& Client, const FGuid& StreamId)
+		static const FConcertReplicationStream* FindExistingStream(const FConcertReplicationClient& Client, const FGuid& StreamId)
 		{
-			return Client.GetStreamDescriptions().FindByPredicate([&StreamId](const FReplicationStreamDescription& Description)
+			return Client.GetStreamDescriptions().FindByPredicate([&StreamId](const FConcertReplicationStream& Description)
 			{
 				return Description.BaseDescription.Identifier == StreamId;
 			});
@@ -25,7 +25,7 @@ namespace UE::ConcertSyncServer::Replication
 		/** Validates that ObjectsToPut writes only to pre-existing streams. */
 		static void ValidatePutObjectsRequestSemantics(const FConcertReplication_ChangeStream_Request& Request, const FConcertReplicationClient& Client, FConcertReplication_ChangeStream_Response& OutResponse)
 		{
-			const auto PutObjectHasEnoughData = [](const FReplicationStreamDescription* ExistingStream, const FSoftObjectPath& ChangedObjectPath, const FConcertReplication_ChangeStream_PutObject& PutObject)
+			const auto PutObjectHasEnoughData = [](const FConcertReplicationStream* ExistingStream, const FSoftObjectPath& ChangedObjectPath, const FConcertReplication_ChangeStream_PutObject& PutObject)
 			{
 				const bool bHasProperties = !PutObject.Properties.ReplicatedProperties.IsEmpty();
 				const bool bHasClassPath = !PutObject.ClassPath.IsNull();
@@ -36,16 +36,16 @@ namespace UE::ConcertSyncServer::Replication
 				return bHasEnoughData;
 			};
 			
-			for (const TPair<FObjectInStreamID, FConcertReplication_ChangeStream_PutObject>& Change : Request.ObjectsToPut)
+			for (const TPair<FConcertObjectInStreamID, FConcertReplication_ChangeStream_PutObject>& Change : Request.ObjectsToPut)
 			{
-				const FObjectInStreamID ChangedObject = Change.Key;
+				const FConcertObjectInStreamID ChangedObject = Change.Key;
 				
 				const FGuid& StreamToModify = ChangedObject.StreamId;
-				const FReplicationStreamDescription* ExistingStream = FindExistingStream(Client, StreamToModify);
+				const FConcertReplicationStream* ExistingStream = FindExistingStream(Client, StreamToModify);
 				const bool bStreamExists = ExistingStream != nullptr;
 				if (bStreamExists)
 				{
-					// FReplicatedObjectInfo must have non-empty values for its members. Hence we must check that
+					// FConcertReplicatedObjectInfo must have non-empty values for its members. Hence we must check that
 					// 1. if creating a new entry, both fields are given
 					// 2. if writing to pre-existing entry, at least one field is given
 					const FConcertReplication_ChangeStream_PutObject& PutObject = Change.Value; 
@@ -67,7 +67,7 @@ namespace UE::ConcertSyncServer::Replication
 		static void ValidateAddedStreamsAreUnique(const FConcertReplication_ChangeStream_Request& Request, const FConcertReplicationClient& Client, FConcertReplication_ChangeStream_Response& OutResponse)
 		{
 			TSet<FGuid> DuplicateEntryDetection;
-			for (const FReplicationStreamDescription& NewStream : Request.StreamsToAdd)
+			for (const FConcertReplicationStream& NewStream : Request.StreamsToAdd)
 			{
 				// StreamsToAdd is invalid if there is already a stream with the same ID registered ...
 				const FGuid& NewStreamId = NewStream.BaseDescription.Identifier;
@@ -89,10 +89,10 @@ namespace UE::ConcertSyncServer::Replication
 		static void LookForAuthorityConflicts(const FConcertReplication_ChangeStream_Request& Request, const FConcertReplicationClient& Client, const FAuthorityManager& AuthorityManager, FConcertReplication_ChangeStream_Response& OutResponse)
 		{
 			const FGuid ClientEndpointId = Client.GetClientEndpointId();
-			for (const TPair<FObjectInStreamID, FConcertReplication_ChangeStream_PutObject>& PutObjectPair : Request.ObjectsToPut)
+			for (const TPair<FConcertObjectInStreamID, FConcertReplication_ChangeStream_PutObject>& PutObjectPair : Request.ObjectsToPut)
 			{
 				// No conflict possible if requesting client does not have authority over the changed object
-				const FReplicatedObjectId ReplicatedObjectInfo { { PutObjectPair.Key }, ClientEndpointId };
+				const FConcertReplicatedObjectId ReplicatedObjectInfo { { PutObjectPair.Key }, ClientEndpointId };
 				if (!AuthorityManager.HasAuthorityToChange(ReplicatedObjectInfo))
 				{
 					continue;
@@ -110,7 +110,7 @@ namespace UE::ConcertSyncServer::Replication
 				AuthorityManager.EnumerateAuthorityConflicts(ReplicatedObjectInfo, &PropertySelection,
 					[&OutResponse, &ReplicatedObjectInfo](const FGuid& ClientId, const FGuid& StreamId, const FConcertPropertyChain& Property)
 					{
-						const FReplicatedObjectId ConflictingObject = { { StreamId, ReplicatedObjectInfo.Object }, ClientId };
+						const FConcertReplicatedObjectId ConflictingObject = { { StreamId, ReplicatedObjectInfo.Object }, ClientId };
 						OutResponse.AuthorityConflicts.Add(ReplicatedObjectInfo, ConflictingObject);
 						
 						UE_LOG(LogConcert, Log, TEXT("Authority conflict with client %s for stream %s for property %s"), *ClientId.ToString(EGuidFormats::Short), *StreamId.ToString(EGuidFormats::Short), *Property.ToString());
@@ -151,7 +151,7 @@ namespace UE::ConcertSyncServer::Replication
 		{
 			// If the client had authority over any objects that were removed by this request, authority must be cleaned up
 			ConcertSyncCore::Replication::ChangeStreamUtils::ForEachObjectLosingAuthority(Request, SendingClient->Get()->GetStreamDescriptions(),
-				[this, &SendingClientId](const FObjectInStreamID& RemovedObject)
+				[this, &SendingClientId](const FConcertObjectInStreamID& RemovedObject)
 				{
 					AuthorityManager->RemoveAuthority({ RemovedObject, SendingClientId});
 					return EBreakBehavior::Continue;
diff --git a/Engine/Plugins/Developer/Concert/ConcertSync/ConcertSyncServer/Source/ConcertSyncServer/Private/Replication/Processing/ServerObjectReplicationReceiver.cpp b/Engine/Plugins/Developer/Concert/ConcertSync/ConcertSyncServer/Source/ConcertSyncServer/Private/Replication/Processing/ServerObjectReplicationReceiver.cpp
--- a/Engine/Plugins/Developer/Concert/ConcertSync/ConcertSyncServer/Source/ConcertSyncServer/Private/Replication/Processing/ServerObjectReplicationReceiver.cpp
+++ b/Engine/Plugins/Developer/Concert/ConcertSync/ConcertSyncServer/Source/ConcertSyncServer/Private/Replication/Processing/ServerObjectReplicationReceiver.cpp
@@ -24,7 +24,7 @@ namespace UE::ConcertSyncServer::Replication
 		const FConcertReplication_ObjectReplicationEvent& ObjectEvent
 		) const
 	{
-		const FReplicatedObjectId ReplicatedObjectInfo { { StreamEvent.StreamId, ObjectEvent.ReplicatedObject }, SessionContext.SourceEndpointId };
+		const FConcertReplicatedObjectId ReplicatedObjectInfo { { StreamEvent.StreamId, ObjectEvent.ReplicatedObject }, SessionContext.SourceEndpointId };
 		return AuthorityManager->HasAuthorityToChange(ReplicatedObjectInfo);
 	}
 }
diff --git a/Engine/Plugins/Developer/Concert/ConcertSync/ConcertSyncServer/Source/ConcertSyncServer/Private/Replication/Processing/ServerReplicationDataQueuer.cpp b/Engine/Plugins/Developer/Concert/ConcertSync/ConcertSyncServer/Source/ConcertSyncServer/Private/Replication/Processing/ServerReplicationDataQueuer.cpp
--- a/Engine/Plugins/Developer/Concert/ConcertSync/ConcertSyncServer/Source/ConcertSyncServer/Private/Replication/Processing/ServerReplicationDataQueuer.cpp
+++ b/Engine/Plugins/Developer/Concert/ConcertSync/ConcertSyncServer/Source/ConcertSyncServer/Private/Replication/Processing/ServerReplicationDataQueuer.cpp
@@ -3,7 +3,6 @@
 #include "ServerReplicationDataQueuer.h"
 
 #include "Replication/ConcertReplicationClient.h"
-#include "Replication/Data/ReplicationStreamDescription.h"
 
 namespace UE::ConcertSyncServer::Replication
 {
@@ -17,7 +16,7 @@ namespace UE::ConcertSyncServer::Replication
 		return Result;
 	}
 
-	bool FServerReplicationDataQueuer::WantsToAcceptObject(const FReplicatedObjectId& Object) const
+	bool FServerReplicationDataQueuer::WantsToAcceptObject(const FConcertReplicatedObjectId& Object) const
 	{
 		// Do not send back the data to the client that generated it
 		const bool bWasSentByThisClient = OwningClientEndpointId == Object.SenderEndpointId;
diff --git a/Engine/Plugins/Developer/Concert/ConcertSync/ConcertSyncServer/Source/ConcertSyncServer/Private/Replication/Processing/ServerReplicationDataQueuer.h b/Engine/Plugins/Developer/Concert/ConcertSync/ConcertSyncServer/Source/ConcertSyncServer/Private/Replication/Processing/ServerReplicationDataQueuer.h
--- a/Engine/Plugins/Developer/Concert/ConcertSync/ConcertSyncServer/Source/ConcertSyncServer/Private/Replication/Processing/ServerReplicationDataQueuer.h
+++ b/Engine/Plugins/Developer/Concert/ConcertSync/ConcertSyncServer/Source/ConcertSyncServer/Private/Replication/Processing/ServerReplicationDataQueuer.h
@@ -24,7 +24,7 @@ namespace UE::ConcertSyncServer::Replication
 		static TSharedRef<FServerReplicationDataQueuer> Make(const FGuid& OwningClientEndpointId, TSharedRef<ConcertSyncCore::FObjectReplicationCache> InReplicationCache);
 		
 		//~ Begin IReplicationCacheUser Interface
-		virtual bool WantsToAcceptObject(const FReplicatedObjectId& Object) const override;
+		virtual bool WantsToAcceptObject(const FConcertReplicatedObjectId& Object) const override;
 		//~ End IReplicationCacheUser Interface
 
 	private:
diff --git a/Engine/Plugins/Developer/Concert/ConcertSync/ConcertSyncServer/Source/ConcertSyncServer/Private/Replication/Util/JoinRequestValidation.cpp b/Engine/Plugins/Developer/Concert/ConcertSync/ConcertSyncServer/Source/ConcertSyncServer/Private/Replication/Util/JoinRequestValidation.cpp
--- a/Engine/Plugins/Developer/Concert/ConcertSync/ConcertSyncServer/Source/ConcertSyncServer/Private/Replication/Util/JoinRequestValidation.cpp
+++ b/Engine/Plugins/Developer/Concert/ConcertSync/ConcertSyncServer/Source/ConcertSyncServer/Private/Replication/Util/JoinRequestValidation.cpp
@@ -2,9 +2,8 @@
 
 #include "JoinRequestValidation.h"
 
-#include "Replication/ConcertReplicationClient.h"
 #include "Replication/Data/ObjectReplicationMap.h"
-#include "Replication/Data/ReplicationStreamDescription.h"
+#include "Replication/Data/ReplicationStream.h"
 #include "Replication/Messages/Handshake.h"
 
 namespace UE::ConcertSyncServer::Replication
@@ -19,7 +18,7 @@ namespace UE::ConcertSyncServer::Replication
 		 */
 		
 		TSet<FGuid> UniqueStreamsFromRequestor;
-		for (const FReplicationStreamDescription& Stream : Request.Streams)
+		for (const FConcertReplicationStream& Stream : Request.Streams)
 		{
 			const FGuid& StreamId = Stream.BaseDescription.Identifier;
 
@@ -31,10 +30,10 @@ namespace UE::ConcertSyncServer::Replication
 			UniqueStreamsFromRequestor.Add(Stream.BaseDescription.Identifier);
 
 			// Make sure the properties are valid
-			for (const TPair<FSoftObjectPath, FReplicatedObjectInfo>& ReplicatedObjectInfo : Stream.BaseDescription.ReplicationMap.ReplicatedObjects)
+			for (const TPair<FSoftObjectPath, FConcertReplicatedObjectInfo>& ReplicatedObjectInfo : Stream.BaseDescription.ReplicationMap.ReplicatedObjects)
 			{
 				const FSoftObjectPath& ObjectPath = ReplicatedObjectInfo.Key;
-				const FReplicatedObjectInfo& ObjectInfo = ReplicatedObjectInfo.Value;
+				const FConcertReplicatedObjectInfo& ObjectInfo = ReplicatedObjectInfo.Value;
 				
 				// Validate class - we cannot validate at this point whether the class really exists.
 				if (!ObjectInfo.ClassPath.IsValid())
@@ -70,12 +69,12 @@ namespace UE::ConcertSyncServer::Replication
 		return { EJoinReplicationErrorCode::Success, TEXT("") };
 	}
 	
-	TTuple<EJoinReplicationErrorCode, FString, TArray<FReplicationStreamDescription>> ValidateRequest(const FConcertReplication_Join_Request& Request)
+	TTuple<EJoinReplicationErrorCode, FString, TArray<FConcertReplicationStream>> ValidateRequest(const FConcertReplication_Join_Request& Request)
 	{
 		if (auto[ErrorCode, ErrorMessage] = CheckInputConstraints(Request)
 			; ErrorCode != EJoinReplicationErrorCode::Success)
 		{
-			return { ErrorCode, ErrorMessage, TArray<FReplicationStreamDescription>{} };
+			return { ErrorCode, ErrorMessage, TArray<FConcertReplicationStream>{} };
 		}
 
 		return { EJoinReplicationErrorCode::Success, TEXT(""), Request.Streams };
diff --git a/Engine/Plugins/Developer/Concert/ConcertSync/ConcertSyncServer/Source/ConcertSyncServer/Private/Replication/Util/JoinRequestValidation.h b/Engine/Plugins/Developer/Concert/ConcertSync/ConcertSyncServer/Source/ConcertSyncServer/Private/Replication/Util/JoinRequestValidation.h
--- a/Engine/Plugins/Developer/Concert/ConcertSync/ConcertSyncServer/Source/ConcertSyncServer/Private/Replication/Util/JoinRequestValidation.h
+++ b/Engine/Plugins/Developer/Concert/ConcertSync/ConcertSyncServer/Source/ConcertSyncServer/Private/Replication/Util/JoinRequestValidation.h
@@ -7,7 +7,7 @@
 
 enum class EJoinReplicationErrorCode : uint8;
 struct FConcertReplication_Join_Request;
-struct FReplicationStreamDescription;
+struct FConcertReplicationStream;
 
 namespace UE::ConcertSyncServer::Replication
 {
@@ -24,5 +24,5 @@ namespace UE::ConcertSyncServer::Replication
 	 * 
 	 * @return The error code for the request and the unpacked stream descriptions
 	 */
-	TTuple<EJoinReplicationErrorCode, FString, TArray<FReplicationStreamDescription>> ValidateRequest(const FConcertReplication_Join_Request& Request);
+	TTuple<EJoinReplicationErrorCode, FString, TArray<FConcertReplicationStream>> ValidateRequest(const FConcertReplication_Join_Request& Request);
 }
diff --git a/Engine/Plugins/Developer/Concert/ConcertSync/ConcertSyncTest/Source/ConcertSyncTest/Private/Replication/Components/FrequencyTrackerTests.cpp b/Engine/Plugins/Developer/Concert/ConcertSync/ConcertSyncTest/Source/ConcertSyncTest/Private/Replication/Components/FrequencyTrackerTests.cpp
--- a/Engine/Plugins/Developer/Concert/ConcertSync/ConcertSyncTest/Source/ConcertSyncTest/Private/Replication/Components/FrequencyTrackerTests.cpp
+++ b/Engine/Plugins/Developer/Concert/ConcertSync/ConcertSyncTest/Source/ConcertSyncTest/Private/Replication/Components/FrequencyTrackerTests.cpp
@@ -32,8 +32,8 @@ namespace UE::ConcertSyncTests
 
 		// 1. Init test data
 		const FSoftObjectPath ReplicatedObject = GetMutableDefault<UObject>();
-		const FObjectInStreamID SpecifiedRateObject { FGuid{ 1, 0, 0, 0 }, ReplicatedObject };
-		const FObjectInStreamID RealtimeObject { FGuid{ 2, 0, 0, 0 }, ReplicatedObject };
+		const FConcertObjectInStreamID SpecifiedRateObject { FGuid{ 1, 0, 0, 0 }, ReplicatedObject };
+		const FConcertObjectInStreamID RealtimeObject { FGuid{ 2, 0, 0, 0 }, ReplicatedObject };
 
 		constexpr FConcertObjectReplicationSettings SpecifiedRateSettings { EConcertObjectReplicationMode::SpecifiedRate, 30 };
 		constexpr FConcertObjectReplicationSettings RealtimeSettings { EConcertObjectReplicationMode::Realtime };
diff --git a/Engine/Plugins/Developer/Concert/ConcertSync/ConcertSyncTest/Source/ConcertSyncTest/Private/Replication/Components/ReplicationCacheTests.cpp b/Engine/Plugins/Developer/Concert/ConcertSync/ConcertSyncTest/Source/ConcertSyncTest/Private/Replication/Components/ReplicationCacheTests.cpp
--- a/Engine/Plugins/Developer/Concert/ConcertSync/ConcertSyncTest/Source/ConcertSyncTest/Private/Replication/Components/ReplicationCacheTests.cpp
+++ b/Engine/Plugins/Developer/Concert/ConcertSync/ConcertSyncTest/Source/ConcertSyncTest/Private/Replication/Components/ReplicationCacheTests.cpp
@@ -60,22 +60,22 @@ namespace UE::ConcertSyncTests
 		public:
 
 			FAutomationTestBase& Test;
-			FObjectInStreamID AllowedObject;
+			FConcertObjectInStreamID AllowedObject;
 			EReplicationCacheTestFlags Flags;
 			TSharedPtr<const FConcertReplication_ObjectReplicationEvent> CachedData;
 
-			FTestReplicationCacheUser(FAutomationTestBase& Test, FObjectInStreamID AllowedObject, EReplicationCacheTestFlags Flags = EReplicationCacheTestFlags::None)
+			FTestReplicationCacheUser(FAutomationTestBase& Test, FConcertObjectInStreamID AllowedObject, EReplicationCacheTestFlags Flags = EReplicationCacheTestFlags::None)
 				: Test(Test)
 				, AllowedObject(MoveTemp(AllowedObject))
 				, Flags(Flags)
 			{}
 			
-			virtual bool WantsToAcceptObject(const FReplicatedObjectId& Object) const override
+			virtual bool WantsToAcceptObject(const FConcertReplicatedObjectId& Object) const override
 			{
 				return !EnumHasAnyFlags(Flags, EReplicationCacheTestFlags::NeverReceive) && AllowedObject == Object;
 			}
 			
-			virtual void OnDataCached(const FReplicatedObjectId& Object, TSharedRef<const FConcertReplication_ObjectReplicationEvent> Data) override
+			virtual void OnDataCached(const FConcertReplicatedObjectId& Object, TSharedRef<const FConcertReplication_ObjectReplicationEvent> Data) override
 			{
 				if (EnumHasAnyFlags(Flags, EReplicationCacheTestFlags::NeverReceive))
 				{
@@ -124,7 +124,7 @@ namespace UE::ConcertSyncTests
 		const FGuid StreamId = FGuid::NewGuid();
 		const FSoftObjectPath ObjectPath(TEXT("/Game/World.World:PersistentLevel.StaticMeshActor0"));
 		const FGuid DummySendingClientId = FGuid::NewGuid();
-		const FReplicatedObjectId ObjectID{ { StreamId, ObjectPath }, DummySendingClientId};
+		const FConcertReplicatedObjectId ObjectID{ { StreamId, ObjectPath }, DummySendingClientId};
 		
 		TSharedRef<ConcertSyncCore::FObjectReplicationCache> Cache = MakeShared<ConcertSyncCore::FObjectReplicationCache>(MakeShared<FTestReplicationFormat>());
 		TSharedRef<FTestReplicationCacheUser> User_NeverConsume = MakeShared<FTestReplicationCacheUser>(*this, ObjectID, EReplicationCacheTestFlags::NeverConsume);
diff --git a/Engine/Plugins/Developer/Concert/ConcertSync/ConcertSyncTest/Source/ConcertSyncTest/Private/Replication/Components/ReplicationChangeStreamUtilsTests.cpp b/Engine/Plugins/Developer/Concert/ConcertSync/ConcertSyncTest/Source/ConcertSyncTest/Private/Replication/Components/ReplicationChangeStreamUtilsTests.cpp
--- a/Engine/Plugins/Developer/Concert/ConcertSync/ConcertSyncTest/Source/ConcertSyncTest/Private/Replication/Components/ReplicationChangeStreamUtilsTests.cpp
+++ b/Engine/Plugins/Developer/Concert/ConcertSync/ConcertSyncTest/Source/ConcertSyncTest/Private/Replication/Components/ReplicationChangeStreamUtilsTests.cpp
@@ -44,11 +44,11 @@ namespace UE::ConcertSyncTests::ChangeStreamUtils
 		
 
 		// 2. Run
-		FObjectReplicationMap Base;
+		FConcertObjectReplicationMap Base;
 		Base.ReplicatedObjects.Add(ChangeProperties, { SceneComponentClass, Selection_RelativeLocation });
 		Base.ReplicatedObjects.Add(ChangeClass, { SceneComponentClass, Selection_RelativeLocation });
 		Base.ReplicatedObjects.Add(RemoveMe, { SceneComponentClass, Selection_RelativeLocation });
-		FObjectReplicationMap Desired;
+		FConcertObjectReplicationMap Desired;
 		Desired.ReplicatedObjects.Add(ChangeProperties, { SceneComponentClass, Selection_RelativeLocationOnlyX });
 		Desired.ReplicatedObjects.Add(ChangeClass, { StaticMeshComponentClass, Selection_RelativeLocation });
 		Desired.ReplicatedObjects.Add(AddMe, { SceneComponentClass, Selection_RelativeLocation });
diff --git a/Engine/Plugins/Developer/Concert/ConcertSync/ConcertSyncTest/Source/ConcertSyncTest/Private/Replication/FrequencyTests.cpp b/Engine/Plugins/Developer/Concert/ConcertSync/ConcertSyncTest/Source/ConcertSyncTest/Private/Replication/FrequencyTests.cpp
--- a/Engine/Plugins/Developer/Concert/ConcertSync/ConcertSyncTest/Source/ConcertSyncTest/Private/Replication/FrequencyTests.cpp
+++ b/Engine/Plugins/Developer/Concert/ConcertSync/ConcertSyncTest/Source/ConcertSyncTest/Private/Replication/FrequencyTests.cpp
@@ -1,6 +1,6 @@
 // Copyright Epic Games, Inc. All Rights Reserved.
 
-#include "Replication/Data/ReplicationStreamDescription.h"
+#include "Replication/Data/ReplicationStream.h"
 #include "Replication/IConcertClientReplicationManager.h"
 #include "Replication/Util/SendReceiveObjectTestBase.h"
 #include "TestReflectionObject.h"
diff --git a/Engine/Plugins/Developer/Concert/ConcertSync/ConcertSyncTest/Source/ConcertSyncTest/Private/Replication/ReplicationAuthorityTests.cpp b/Engine/Plugins/Developer/Concert/ConcertSync/ConcertSyncTest/Source/ConcertSyncTest/Private/Replication/ReplicationAuthorityTests.cpp
--- a/Engine/Plugins/Developer/Concert/ConcertSync/ConcertSyncTest/Source/ConcertSyncTest/Private/Replication/ReplicationAuthorityTests.cpp
+++ b/Engine/Plugins/Developer/Concert/ConcertSync/ConcertSyncTest/Source/ConcertSyncTest/Private/Replication/ReplicationAuthorityTests.cpp
@@ -376,19 +376,19 @@ namespace UE::ConcertSyncTests::Replication::Authority
 		{
 		public:
 			
-			FSharedReplicationStreamDescription RequestingClientStream;
-			FSharedReplicationStreamDescription ExistingClientStream;
+			FConcertBaseStreamInfo RequestingClientStream;
+			FConcertBaseStreamInfo ExistingClientStream;
 			bool bRequestingClientHasAuthority;
 			bool bExistingClientHasAuthority;
 
-			FTestGroundTruth(const FSharedReplicationStreamDescription& RequestingClientStream, const FSharedReplicationStreamDescription& ExistingClientStream, bool bRequestingClientHasAuthority, bool bExistingClientHasAuthority)
+			FTestGroundTruth(const FConcertBaseStreamInfo& RequestingClientStream, const FConcertBaseStreamInfo& ExistingClientStream, bool bRequestingClientHasAuthority, bool bExistingClientHasAuthority)
 				: RequestingClientStream(RequestingClientStream)
 				, ExistingClientStream(ExistingClientStream)
 				, bRequestingClientHasAuthority(bRequestingClientHasAuthority)
 				, bExistingClientHasAuthority(bExistingClientHasAuthority)
 			{}
 
-			virtual void ForEachStream(const FGuid& ClientEndpointId, TFunctionRef<EBreakBehavior(const FGuid& StreamId, const FObjectReplicationMap& ReplicationMap)> Callback) const override
+			virtual void ForEachStream(const FGuid& ClientEndpointId, TFunctionRef<EBreakBehavior(const FGuid& StreamId, const FConcertObjectReplicationMap& ReplicationMap)> Callback) const override
 			{
 				if (ClientEndpointId == RequestingClientId)
 				{
@@ -420,12 +420,12 @@ namespace UE::ConcertSyncTests::Replication::Authority
 			}
 		};
 		
-		FSharedReplicationStreamDescription MakeStream(const FGuid& StreamId, const FSoftObjectPath ObjectPath, TArray<FConcertPropertyChain> Properties)
+		FConcertBaseStreamInfo MakeStream(const FGuid& StreamId, const FSoftObjectPath ObjectPath, TArray<FConcertPropertyChain> Properties)
 		{
-			FObjectReplicationMap ExistingReplicationMap;
+			FConcertObjectReplicationMap ExistingReplicationMap;
 			const FConcertPropertySelection Selection{ Properties};
 			ExistingReplicationMap.ReplicatedObjects.Add(ObjectPath, { {}, Selection });
-			const FSharedReplicationStreamDescription ExistingClientStream { StreamId, ExistingReplicationMap };
+			const FConcertBaseStreamInfo ExistingClientStream { StreamId, ExistingReplicationMap };
 			return ExistingClientStream;
 		};
 	}
@@ -442,9 +442,9 @@ namespace UE::ConcertSyncTests::Replication::Authority
 		FConcertPropertyChain Float = *FConcertPropertyChain::CreateFromPath(*UTestReflectionObject::StaticClass(), { TEXT("Float") });
 		FConcertPropertyChain Vector = *FConcertPropertyChain::CreateFromPath(*UTestReflectionObject::StaticClass(), { TEXT("Vector") });
 		FConcertPropertyChain VectorX = *FConcertPropertyChain::CreateFromPath(*UTestReflectionObject::StaticClass(), { TEXT("Vector"), TEXT("X") });
-		const FSharedReplicationStreamDescription ClientStream_Empty = MakeStream(TestStreamId, ObjectPath, {});
-		const FSharedReplicationStreamDescription ClientStream_FloatOnly = MakeStream(TestStreamId, ObjectPath, { Float });
-		const FSharedReplicationStreamDescription ClientStream_VectorOnly = MakeStream(TestStreamId, ObjectPath, { Vector, VectorX });
+		const FConcertBaseStreamInfo ClientStream_Empty = MakeStream(TestStreamId, ObjectPath, {});
+		const FConcertBaseStreamInfo ClientStream_FloatOnly = MakeStream(TestStreamId, ObjectPath, { Float });
+		const FConcertBaseStreamInfo ClientStream_VectorOnly = MakeStream(TestStreamId, ObjectPath, { Vector, VectorX });
 
 		// Requestor tries to stream Float, then Vector
 		// Existing is replicating Float
diff --git a/Engine/Plugins/Developer/Concert/ConcertSync/ConcertSyncTest/Source/ConcertSyncTest/Private/Replication/ReplicationHandshakeTests.cpp b/Engine/Plugins/Developer/Concert/ConcertSync/ConcertSyncTest/Source/ConcertSyncTest/Private/Replication/ReplicationHandshakeTests.cpp
--- a/Engine/Plugins/Developer/Concert/ConcertSync/ConcertSyncTest/Source/ConcertSyncTest/Private/Replication/ReplicationHandshakeTests.cpp
+++ b/Engine/Plugins/Developer/Concert/ConcertSync/ConcertSyncTest/Source/ConcertSyncTest/Private/Replication/ReplicationHandshakeTests.cpp
@@ -2,8 +2,7 @@
 
 #include "Util/ClientServerCommunicationTest.h"
 
-#include "Replication/Data/ReplicationClientDescription.h"
-#include "Replication/Data/ReplicationStreamDescription.h"
+#include "Replication/Data/ReplicationStream.h"
 #include "Replication/IConcertClientReplicationBridge.h"
 #include "Replication/IConcertClientReplicationManager.h"
 #include "Replication/IConcertServerReplicationManager.h"
@@ -49,19 +48,18 @@ namespace UE::ConcertSyncTests::Replication::Handshake
 		const TSharedPtr<IConcertClientSession>& ClientSession_Secondary = Client_Secondary.ClientSessionMock;
 		const TSharedRef<IConcertClientReplicationManager> ClientReplicationManager_Secondary = CreateClientReplicationManager(ClientSession_Secondary.ToSharedRef(), &BridgeMock_Secondary.Get());
 
-		// Stream info
-		FReplicationClientDescription ClientDescription;
-		FReplicationStreamDescription StreamDescription;
-		FReplicatedObjectInfo StaticMeshComponentInfo;
+		// Stream 
+		FConcertReplicationStream StreamDescription;
+		FConcertReplicatedObjectInfo StaticMeshComponentInfo;
 		StaticMeshComponentInfo.ClassPath = UStaticMeshComponent::StaticClass();
 		StaticMeshComponentInfo.PropertySelection.ReplicatedProperties.Add(ForcedLodModelProperty);
 		StreamDescription.BaseDescription.ReplicationMap.ReplicatedObjects.Add(PathToSomeActorComponent, StaticMeshComponentInfo);
 
-		FReplicationStreamDescription DuplicateStreamDescription = StreamDescription;
+		FConcertReplicationStream DuplicateStreamDescription = StreamDescription;
 		DuplicateStreamDescription.BaseDescription.Identifier = FGuid::NewGuid();
 		
-		FReplicationStreamDescription InvalidClassStreamDescription = StreamDescription;
-		FReplicatedObjectInfo MissingClassInfo = StaticMeshComponentInfo;
+		FConcertReplicationStream InvalidClassStreamDescription = StreamDescription;
+		FConcertReplicatedObjectInfo MissingClassInfo = StaticMeshComponentInfo;
 		MissingClassInfo.ClassPath.Reset();
 		InvalidClassStreamDescription.BaseDescription.ReplicationMap.ReplicatedObjects.Add(PathToSomeActorComponent, MissingClassInfo);
 
@@ -73,21 +71,21 @@ namespace UE::ConcertSyncTests::Replication::Handshake
 		// However malicious users or those with custom C++ logic can send whatever they want.
 		
 		// 2.1.1 Duplicate properties
-		FReplicationStreamDescription Invalid_DoublePropertyDescription = StreamDescription;
+		FConcertReplicationStream Invalid_DoublePropertyDescription = StreamDescription;
 		Invalid_DoublePropertyDescription.BaseDescription.ReplicationMap.ReplicatedObjects[PathToSomeActorComponent].PropertySelection.ReplicatedProperties.Add(ForcedLodModelProperty);
-		ClientReplicationManager_Primary->JoinReplicationSession({ ClientDescription, { Invalid_DoublePropertyDescription } })
+		ClientReplicationManager_Primary->JoinReplicationSession({ { Invalid_DoublePropertyDescription } })
 			.Next([&](const FJoinReplicatedSessionResult& Result)
 			{
 				TestTrue(TEXT("Cannot contain same properties twice"), Result.ErrorCode == EJoinReplicationErrorCode::DuplicateProperty);
 			});
 		// 2.1.2 Duplicate stream identifier
-		ClientReplicationManager_Primary->JoinReplicationSession({ ClientDescription, { StreamDescription, StreamDescription } })
+		ClientReplicationManager_Primary->JoinReplicationSession({ { StreamDescription, StreamDescription } })
 			.Next([&](const FJoinReplicatedSessionResult& Result)
 			{
 				TestTrue(TEXT("Cannot contain stream ID twice"), Result.ErrorCode == EJoinReplicationErrorCode::DuplicateStreamId);
 			});
 		// 2.1.3 Missing class path
-		ClientReplicationManager_Primary->JoinReplicationSession({ ClientDescription, { InvalidClassStreamDescription } })
+		ClientReplicationManager_Primary->JoinReplicationSession({ { InvalidClassStreamDescription } })
 			.Next([&](const FJoinReplicatedSessionResult& Result)
 			{
 				TestTrue(TEXT("Cannot contain null classes"), Result.ErrorCode == EJoinReplicationErrorCode::InvalidClass);
@@ -96,21 +94,21 @@ namespace UE::ConcertSyncTests::Replication::Handshake
 
 		// 2.2 Real workflow cases
 		// 2.2.1 Valid join
-		ClientReplicationManager_Primary->JoinReplicationSession({ ClientDescription, { StreamDescription } })
+		ClientReplicationManager_Primary->JoinReplicationSession({ { StreamDescription } })
 			.Next([&](const FJoinReplicatedSessionResult& Result)
 			{
 				TestTrue(TEXT("Join with valid args"), Result.ErrorCode == EJoinReplicationErrorCode::Success);
 			});
 		// 2.2.2 No joining twice
 		AddExpectedError(TEXT("JoinReplicationSession requested while already in a session"), EAutomationExpectedErrorFlags::Contains); // Not pretty, but otherwise this test fails due to logged warning
-		ClientReplicationManager_Primary->JoinReplicationSession({ ClientDescription, { StreamDescription } })
+		ClientReplicationManager_Primary->JoinReplicationSession({ { StreamDescription } })
 			.Next([&](const FJoinReplicatedSessionResult& Result)
 			{
 				TestTrue(TEXT("Cannot join twice"), Result.ErrorCode == EJoinReplicationErrorCode::AlreadyInSession);
 			});
 		// 2.2.3 Rejoining
 		ClientReplicationManager_Primary->LeaveReplicationSession();
-		ClientReplicationManager_Primary->JoinReplicationSession({ ClientDescription, { StreamDescription } })
+		ClientReplicationManager_Primary->JoinReplicationSession({ { StreamDescription } })
 			.Next([&](const FJoinReplicatedSessionResult& Result)
 			{
 				TestTrue(TEXT("Re-join session"), Result.ErrorCode == EJoinReplicationErrorCode::Success);
@@ -120,8 +118,8 @@ namespace UE::ConcertSyncTests::Replication::Handshake
 		// For context: UX-wise it should be as easy as possible for clients to join.
 		// This is why we allow clients to join with overlapping properties and stream IDs, which used to be rejected by the handshake.
 		// Now, they must take authority over objects before sending. See authority tests.
-		FReplicationStreamDescription DuplicateProperties = StreamDescription;
-		ClientReplicationManager_Secondary->JoinReplicationSession({ ClientDescription, { StreamDescription } })
+		FConcertReplicationStream DuplicateProperties = StreamDescription;
+		ClientReplicationManager_Secondary->JoinReplicationSession({ { StreamDescription } })
 			.Next([&](const FJoinReplicatedSessionResult& Result)
 			{
 				TestTrue(TEXT("2nd client can overlap stream identifier and properties of 1st client"), Result.ErrorCode == EJoinReplicationErrorCode::Success);
diff --git a/Engine/Plugins/Developer/Concert/ConcertSync/ConcertSyncTest/Source/ConcertSyncTest/Private/Replication/ReplicationSendReceiveTests.cpp b/Engine/Plugins/Developer/Concert/ConcertSync/ConcertSyncTest/Source/ConcertSyncTest/Private/Replication/ReplicationSendReceiveTests.cpp
--- a/Engine/Plugins/Developer/Concert/ConcertSync/ConcertSyncTest/Source/ConcertSyncTest/Private/Replication/ReplicationSendReceiveTests.cpp
+++ b/Engine/Plugins/Developer/Concert/ConcertSync/ConcertSyncTest/Source/ConcertSyncTest/Private/Replication/ReplicationSendReceiveTests.cpp
@@ -2,7 +2,7 @@
 
 #include "Util/SendReceiveObjectTestBase.h"
 
-#include "Replication/Data/ReplicationStreamDescription.h"
+#include "Replication/Data/ReplicationStream.h"
 #include "Replication/IConcertClientReplicationManager.h"
 #include "Replication/IConcertServerReplicationManager.h"
 #include "Replication/Messages/ObjectReplication.h"
@@ -76,7 +76,7 @@ namespace UE::ConcertSyncTests::Replication::SendReceiveFlow
 		VectorOnlySelection.ReplicatedProperties.Add(*VectorPropertyChain);
 		
 		FConcertReplication_ChangeStream_Request Request;
-		Request.ObjectsToPut.Add(FObjectInStreamID{ Test.SenderStreamId, Test.TestObject }, FConcertReplication_ChangeStream_PutObject{ VectorOnlySelection });
+		Request.ObjectsToPut.Add(FConcertObjectInStreamID{ Test.SenderStreamId, Test.TestObject }, FConcertReplication_ChangeStream_PutObject{ VectorOnlySelection });
 		bool bReceivedChangeStreamResponse = false;
 		Test.ClientReplicationManager_Sender->ChangeStream(Request)
 			.Next([&Test, &bReceivedChangeStreamResponse](FConcertReplication_ChangeStream_Response&& Response)
@@ -118,18 +118,18 @@ namespace UE::ConcertSyncTests::Replication::SendReceiveFlow
 			ConcertSyncClient::Replication::FJoinReplicatedSessionArgs SenderJoinArgs;
 			const UClass& TestClass = *TestObject->GetClass();
 
-			FReplicatedObjectInfo FloatProperties {&TestClass};
+			FConcertReplicatedObjectInfo FloatProperties {&TestClass};
 			FloatProperties.PropertySelection.ReplicatedProperties.Add(*FConcertPropertyChain::CreateFromPath(TestClass, { GET_MEMBER_NAME_CHECKED(UTestReflectionObject, Float) }));
-			FReplicatedObjectInfo VectorProperties { &TestClass };
+			FConcertReplicatedObjectInfo VectorProperties { &TestClass };
 			VectorProperties.PropertySelection.ReplicatedProperties.Add(*FConcertPropertyChain::CreateFromPath(TestClass, { GET_MEMBER_NAME_CHECKED(UTestReflectionObject, Vector) }));
 
-			FReplicationStreamDescription FloatStream;
+			FConcertReplicationStream FloatStream;
 			FloatStream.BaseDescription.Identifier = FloatStreamId;
 			FloatStream.BaseDescription.ReplicationMap.ReplicatedObjects.Add(TestObject, FloatProperties);
 			// Use realtime replication, otherwise we'll have to pass time for the frequency system
 			FloatStream.BaseDescription.FrequencySettings.Defaults.ReplicationMode = EConcertObjectReplicationMode::Realtime;
 			
-			FReplicationStreamDescription VectorStream;
+			FConcertReplicationStream VectorStream;
 			VectorStream.BaseDescription.Identifier = VectorStreamId;
 			VectorStream.BaseDescription.ReplicationMap.ReplicatedObjects.Add(TestObject, VectorProperties);
 			// Use realtime replication, otherwise we'll have to pass time for the frequency system
diff --git a/Engine/Plugins/Developer/Concert/ConcertSync/ConcertSyncTest/Source/ConcertSyncTest/Private/Replication/StreamRequestTests_Change.cpp b/Engine/Plugins/Developer/Concert/ConcertSync/ConcertSyncTest/Source/ConcertSyncTest/Private/Replication/StreamRequestTests_Change.cpp
--- a/Engine/Plugins/Developer/Concert/ConcertSync/ConcertSyncTest/Source/ConcertSyncTest/Private/Replication/StreamRequestTests_Change.cpp
+++ b/Engine/Plugins/Developer/Concert/ConcertSync/ConcertSyncTest/Source/ConcertSyncTest/Private/Replication/StreamRequestTests_Change.cpp
@@ -2,7 +2,7 @@
 
 #include "Util/ClientServerCommunicationTest.h"
 
-#include "Replication/Data/ReplicationStreamDescription.h"
+#include "Replication/Data/ReplicationStream.h"
 #include "Replication/IConcertClientReplicationManager.h"
 #include "TestReflectionObject.h"
 #include "Util/ChangeStreamsTestBase.h"
@@ -90,7 +90,7 @@ namespace UE::ConcertSyncTests::Replication::Stream
 		FConcertReplication_ChangeStream_Request AppendFloatRequest;
 		FConcertPropertySelection& NewSelection = GetPropertySelection(VectorFloatStream, *TestObject);
 		AddFloatProperty(NewSelection);
-		AppendFloatRequest.ObjectsToPut.Add(FObjectInStreamID{ VectorFloatStreamID, TestObject }, FConcertReplication_ChangeStream_PutObject{ NewSelection });
+		AppendFloatRequest.ObjectsToPut.Add(FConcertObjectInStreamID{ VectorFloatStreamID, TestObject }, FConcertReplication_ChangeStream_PutObject{ NewSelection });
 		ChangeStreamForSenderClientAndValidate(TEXT("AppendFloatRequest"), AppendFloatRequest, { FloatStream.BaseDescription, VectorFloatStream.BaseDescription });
 		
 		return true;
@@ -153,7 +153,7 @@ namespace UE::ConcertSyncTests::Replication::Stream
 		TestTrue(TEXT("Receiver > Received response taking authority"), bReceiverTookAuthority);
 
 		// 2.5 Receiver is rejected from modifying the stream to include the Float property because the Sender has authority over it
-		const FObjectInStreamID TestObjectInReceiverStreamId { ReceiverStreamID, TestObject };
+		const FConcertObjectInStreamID TestObjectInReceiverStreamId { ReceiverStreamID, TestObject };
 		FConcertReplication_ChangeStream_Request AddFloatToStreamRequest;
 		FConcertPropertySelection NewSelection = GetPropertySelection(ReceiverStream, *TestObject);
 		AddFloatProperty(NewSelection);
@@ -168,7 +168,7 @@ namespace UE::ConcertSyncTests::Replication::Stream
 				bReceivedResponseAddingFloat = true;
 				TestTrue(TEXT("ErrorCode == Handled"), Response.ErrorCode == EReplicationResponseErrorCode::Handled);
 				TestEqual(TEXT("Append Float > 1 conflict"), Response.AuthorityConflicts.Num(), 1);
-				if (const FReplicatedObjectId* ConflictingObject = Response.AuthorityConflicts.Find(TestObjectInReceiverStreamId))
+				if (const FConcertReplicatedObjectId* ConflictingObject = Response.AuthorityConflicts.Find(TestObjectInReceiverStreamId))
 				{
 					TestEqual(TEXT("Append float > Conflict > Sender Stream correct"), ConflictingObject->StreamId, SenderStreamID);
 					TestEqual(TEXT("Append float > Conflict > Object correct"), ConflictingObject->Object, FSoftObjectPath(TestObject));
@@ -205,7 +205,7 @@ namespace UE::ConcertSyncTests::Replication::Stream
 		FConcertReplication_ChangeStream_Request InvalidRequest;
 		FConcertPropertySelection NewSelection = GetPropertySelection(Stream, *TestObject);
 		AddFloatProperty(NewSelection);
-		InvalidRequest.ObjectsToPut.Add(FObjectInStreamID{ StreamId, TestObject}, FConcertReplication_ChangeStream_PutObject{ NewSelection });
+		InvalidRequest.ObjectsToPut.Add(FConcertObjectInStreamID{ StreamId, TestObject}, FConcertReplication_ChangeStream_PutObject{ NewSelection });
 		InvalidRequest.StreamsToAdd.Add(Stream); // This will make it fail due to pre-existing stream ID
 
 		// Server logs a warning when rejecting - avoid the test being marked with a warning.
@@ -284,7 +284,7 @@ namespace UE::ConcertSyncTests::Replication::Stream
     	ClientReplicationManager_Sender->ChangeStream(Request);
 
     	// 3. Check that the local client's server prediction was reverted
-    	TArray<FReplicationStreamDescription> Streams = ClientReplicationManager_Sender->GetRegisteredStreams();
+    	TArray<FConcertReplicationStream> Streams = ClientReplicationManager_Sender->GetRegisteredStreams();
     	if (Streams.IsEmpty())
     	{
     		AddError(TEXT("Stream change was not reverted"));
diff --git a/Engine/Plugins/Developer/Concert/ConcertSync/ConcertSyncTest/Source/ConcertSyncTest/Private/Replication/StreamRequestTests_Frequency.cpp b/Engine/Plugins/Developer/Concert/ConcertSync/ConcertSyncTest/Source/ConcertSyncTest/Private/Replication/StreamRequestTests_Frequency.cpp
--- a/Engine/Plugins/Developer/Concert/ConcertSync/ConcertSyncTest/Source/ConcertSyncTest/Private/Replication/StreamRequestTests_Frequency.cpp
+++ b/Engine/Plugins/Developer/Concert/ConcertSync/ConcertSyncTest/Source/ConcertSyncTest/Private/Replication/StreamRequestTests_Frequency.cpp
@@ -1,6 +1,6 @@
 // Copyright Epic Games, Inc. All Rights Reserved.
 
-#include "Replication/Data/ReplicationStreamDescription.h"
+#include "Replication/Data/ReplicationStream.h"
 #include "Replication/IConcertClientReplicationManager.h"
 #include "TestReflectionObject.h"
 #include "Util/ChangeStreamsTestBase.h"
@@ -29,7 +29,7 @@ namespace UE::ConcertSyncTests::Replication::Stream
 		UTestReflectionObject* MainTestObject = NewObject<UTestReflectionObject>(GetTransientPackage());
 		UTestReflectionObject* SecondaryTestObject = NewObject<UTestReflectionObject>(GetTransientPackage());
 		FGuid FloatStreamId;
-		FReplicationStreamDescription FloatStream;
+		FConcertReplicationStream FloatStream;
 		Tie(FloatStreamId, FloatStream) = CreateFloatPropertyStream(*MainTestObject);
 		SenderArgs.Streams = { FloatStream };
 		
@@ -37,7 +37,7 @@ namespace UE::ConcertSyncTests::Replication::Stream
 		// We'll intentionally be causing errors below that cause a warning to be logged.
 		AddExpectedError(TEXT("Rejecting ChangeStream"), EAutomationExpectedErrorFlags::Contains, 2);
 
-		FSharedReplicationStreamDescription ExpectedStream = FloatStream.BaseDescription;
+		FConcertBaseStreamInfo ExpectedStream = FloatStream.BaseDescription;
 		
 		// Change defaults to be SpecifiedRate at 42 update rate.
 		{
@@ -108,7 +108,7 @@ namespace UE::ConcertSyncTests::Replication::Stream
 			FrequencyChange.OverridesToAdd.Add(SecondaryTestObject, { EConcertObjectReplicationMode::SpecifiedRate, 101 });
 			ChangeRequest.FrequencyChanges.Add(FloatStreamId, FrequencyChange);
 
-			FReplicatedObjectInfo& ExpectedSecondary = ExpectedStream.ReplicationMap.ReplicatedObjects.Add(SecondaryTestObject);
+			FConcertReplicatedObjectInfo& ExpectedSecondary = ExpectedStream.ReplicationMap.ReplicatedObjects.Add(SecondaryTestObject);
 			AddFloatProperty(ExpectedSecondary.PropertySelection);
 			ExpectedSecondary.ClassPath = SecondaryTestObject->GetClass();
 			ExpectedStream.FrequencySettings.ObjectOverrides.Add(SecondaryTestObject, { EConcertObjectReplicationMode::SpecifiedRate, 101 });
@@ -132,7 +132,7 @@ namespace UE::ConcertSyncTests::Replication::Stream
 		SenderArgs.Streams = { FloatStream };
 		SetUpClientAndServer();
 
-		FSharedReplicationStreamDescription ExpectedStream = FloatStream.BaseDescription;
+		FConcertBaseStreamInfo ExpectedStream = FloatStream.BaseDescription;
 		// Put SpecifiedRate 21
 		{
 			FConcertReplication_ChangeStream_Request ChangeRequest;
@@ -182,7 +182,7 @@ namespace UE::ConcertSyncTests::Replication::Stream
 		SenderArgs.Streams = { FloatStream };
 		
 		SetUpClientAndServer();
-		FSharedReplicationStreamDescription ExpectedStream = FloatStream.BaseDescription;
+		FConcertBaseStreamInfo ExpectedStream = FloatStream.BaseDescription;
 		
 		// Register SecondaryObject with stream
 		{
@@ -193,7 +193,7 @@ namespace UE::ConcertSyncTests::Replication::Stream
 			PutObject.ClassPath = SecondaryTestObject->GetClass();
 			ChangeRequest.ObjectsToPut.Add({ FloatStreamId, SecondaryTestObject }, PutObject);
 
-			ExpectedStream.ReplicationMap.ReplicatedObjects.Add(SecondaryTestObject, FReplicatedObjectInfo{ PutObject.ClassPath, PutObject.Properties });
+			ExpectedStream.ReplicationMap.ReplicatedObjects.Add(SecondaryTestObject, FConcertReplicatedObjectInfo{ PutObject.ClassPath, PutObject.Properties });
 			ChangeStreamForSenderClientAndValidate(TEXT("Add MainTestObject & SecondaryObject"), ChangeRequest, { ExpectedStream });
 		}
 		// Put frequency for MainTestObject & SecondaryObject
@@ -236,7 +236,7 @@ namespace UE::ConcertSyncTests::Replication::Stream
 		SenderArgs.Streams = { FloatStream };
 		SetUpClientAndServer();
 		
-		FSharedReplicationStreamDescription ExpectedStream = FloatStream.BaseDescription;
+		FConcertBaseStreamInfo ExpectedStream = FloatStream.BaseDescription;
 		// Register SecondaryObject with stream
 		{
 			FConcertReplication_ChangeStream_Request ChangeRequest;
@@ -246,7 +246,7 @@ namespace UE::ConcertSyncTests::Replication::Stream
 			PutObject.ClassPath = SecondaryTestObject->GetClass();
 			ChangeRequest.ObjectsToPut.Add({ FloatStreamId, SecondaryTestObject }, PutObject);
 
-			ExpectedStream.ReplicationMap.ReplicatedObjects.Add(SecondaryTestObject, FReplicatedObjectInfo{ PutObject.ClassPath, PutObject.Properties });
+			ExpectedStream.ReplicationMap.ReplicatedObjects.Add(SecondaryTestObject, FConcertReplicatedObjectInfo{ PutObject.ClassPath, PutObject.Properties });
 			ChangeStreamForSenderClientAndValidate(TEXT("Add MainTestObject & SecondaryObject"), ChangeRequest, { ExpectedStream });
 		}
 		
@@ -293,7 +293,7 @@ namespace UE::ConcertSyncTests::Replication::Stream
 		SenderArgs.Streams = { FloatStream };
 		SetUpClientAndServer();
 
-		FSharedReplicationStreamDescription ExpectedStream = FloatStream.BaseDescription;
+		FConcertBaseStreamInfo ExpectedStream = FloatStream.BaseDescription;
 		// Put SpecifiedRate 25, then add 50
 		{
 			FConcertReplication_ChangeStream_Request ChangeRequest;
@@ -321,7 +321,7 @@ namespace UE::ConcertSyncTests::Replication::Stream
 		UTestReflectionObject* MainTestObject = NewObject<UTestReflectionObject>(GetTransientPackage());
 		UTestReflectionObject* SecondaryTestObject = NewObject<UTestReflectionObject>(GetTransientPackage());
 		FGuid FloatStreamId;
-		FReplicationStreamDescription FloatStream;
+		FConcertReplicationStream FloatStream;
 		Tie(FloatStreamId, FloatStream) = CreateFloatPropertyStream(*MainTestObject);
 		SenderArgs.Streams = { FloatStream };
 		
@@ -330,7 +330,7 @@ namespace UE::ConcertSyncTests::Replication::Stream
 		AddExpectedError(TEXT("Rejecting ChangeStream"),EAutomationExpectedErrorFlags::Contains, 2);
 
 		// Change defaults to be SpecifiedRate at 42 update rate.
-		FSharedReplicationStreamDescription ExpectedStream = FloatStream.BaseDescription;
+		FConcertBaseStreamInfo ExpectedStream = FloatStream.BaseDescription;
 		
 		// Cannot change frequency of unregistered stream
 		{
diff --git a/Engine/Plugins/Developer/Concert/ConcertSync/ConcertSyncTest/Source/ConcertSyncTest/Private/Replication/StreamRequestTests_Query.cpp b/Engine/Plugins/Developer/Concert/ConcertSync/ConcertSyncTest/Source/ConcertSyncTest/Private/Replication/StreamRequestTests_Query.cpp
--- a/Engine/Plugins/Developer/Concert/ConcertSync/ConcertSyncTest/Source/ConcertSyncTest/Private/Replication/StreamRequestTests_Query.cpp
+++ b/Engine/Plugins/Developer/Concert/ConcertSync/ConcertSyncTest/Source/ConcertSyncTest/Private/Replication/StreamRequestTests_Query.cpp
@@ -2,7 +2,7 @@
 
 #include "Util/ClientServerCommunicationTest.h"
 
-#include "Replication/Data/ReplicationStreamDescription.h"
+#include "Replication/Data/ReplicationStream.h"
 #include "Replication/IConcertClientReplicationManager.h"
 #include "TestReflectionObject.h"
 #include "Util/ChangeStreamsTestBase.h"
@@ -36,7 +36,7 @@ namespace UE::ConcertSyncTests::Replication::Stream
 		{
 			TestTrue(TEXT("ErrorCode == Handled"), Response.ErrorCode == EReplicationResponseErrorCode::Handled);
 			
-			const FReplicationClientQueriedInfo* Info = Response.ClientInfo.Find(SenderEndpointId);
+			const FConcertQueriedClientInfo* Info = Response.ClientInfo.Find(SenderEndpointId);
 			if (!Info)
 			{
 				AddError(TEXT("No info about sending client!"));
@@ -50,11 +50,11 @@ namespace UE::ConcertSyncTests::Replication::Stream
 				return;
 			}
 				
-			const FSharedReplicationStreamDescription& StreamDescription = Info->Streams[0];
+			const FConcertBaseStreamInfo& StreamDescription = Info->Streams[0];
 			TestEqual(TEXT("StreamDescription->Identifier correct"), StreamDescription.Identifier, SenderStreamId);
 			TestEqual(TEXT("StreamDescription->ReplicationMap has exactly 1 object"), StreamDescription.ReplicationMap.ReplicatedObjects.Num(), 1);
 
-			const FObjectReplicationMap ExpectedReplicationMap = CreateSenderArgs().Streams[0].BaseDescription.ReplicationMap;
+			const FConcertObjectReplicationMap ExpectedReplicationMap = CreateSenderArgs().Streams[0].BaseDescription.ReplicationMap;
 			TestEqual(TEXT("Registered and reported replication maps match"), StreamDescription.ReplicationMap, ExpectedReplicationMap);
 		};
 		
@@ -67,7 +67,7 @@ namespace UE::ConcertSyncTests::Replication::Stream
 				TestTrue(TEXT("ErrorCode == Handled"), Response.ErrorCode == EReplicationResponseErrorCode::Handled);
 				
 				TestReplicationMapContent(Response);
-				if (const FReplicationClientQueriedInfo* Info = Response.ClientInfo.Find(SenderEndpointId))
+				if (const FConcertQueriedClientInfo* Info = Response.ClientInfo.Find(SenderEndpointId))
 				{
 					TestEqual(TEXT("Contains no authority data"), Info->Authority.Num(), 0);
 				}
@@ -93,7 +93,7 @@ namespace UE::ConcertSyncTests::Replication::Stream
 				TestTrue(TEXT("ErrorCode == Handled"), Response.ErrorCode == EReplicationResponseErrorCode::Handled);
 				TestReplicationMapContent(Response);
 				
-				const FReplicationClientQueriedInfo* Info = Response.ClientInfo.Find(SenderEndpointId);
+				const FConcertQueriedClientInfo* Info = Response.ClientInfo.Find(SenderEndpointId);
 				if (!Info || Info->Authority.IsEmpty())
 				{
 					AddError(TEXT("Missing authority data"));
@@ -101,7 +101,7 @@ namespace UE::ConcertSyncTests::Replication::Stream
 				}
 
 				TestEqual(TEXT("Exactly 1 authority stream"), Info->Authority.Num(), 1);
-				const FReplicationAuthorityInfo& AuthorityInfo = Info->Authority[0];
+				const FConcertAuthorityClientInfo& AuthorityInfo = Info->Authority[0];
 				TestEqual(TEXT("Authority stream ID matches registered stream ID"), AuthorityInfo.StreamId, SenderStreamId);
 				TestEqual(TEXT("Has authority over exactly 1 object"), AuthorityInfo.AuthoredObjects.Num(), 1);
 				TestTrue(TEXT("Has authority over registered object"), AuthorityInfo.AuthoredObjects.Contains(TestObjectPath));
@@ -117,7 +117,7 @@ namespace UE::ConcertSyncTests::Replication::Stream
 			{
 				bReceivedResponse_SkipStreamInfo = true;
 				TestTrue(TEXT("ErrorCode == Handled"), Response.ErrorCode == EReplicationResponseErrorCode::Handled);
-				if (const FReplicationClientQueriedInfo* Info = Response.ClientInfo.Find(SenderEndpointId))
+				if (const FConcertQueriedClientInfo* Info = Response.ClientInfo.Find(SenderEndpointId))
 				{
 					TestEqual(TEXT("SkipStreamInfo > No stream data"), Info->Streams.Num(), 0);
 				}
@@ -131,7 +131,7 @@ namespace UE::ConcertSyncTests::Replication::Stream
 			.Next([this, &SenderEndpointId, &TestObjectPath, &bReceivedResponse_SkipProperties](FConcertReplication_QueryReplicationInfo_Response&& Response) mutable
 			{
 				bReceivedResponse_SkipProperties = true;
-				const FReplicationClientQueriedInfo* Info = Response.ClientInfo.Find(SenderEndpointId);
+				const FConcertQueriedClientInfo* Info = Response.ClientInfo.Find(SenderEndpointId);
 				TestTrue(TEXT("ErrorCode == Handled"), Response.ErrorCode == EReplicationResponseErrorCode::Handled);
 				if (!Info || Info->Streams.IsEmpty())
 				{
@@ -139,7 +139,7 @@ namespace UE::ConcertSyncTests::Replication::Stream
 					return;
 				}
 
-				const FReplicatedObjectInfo* ObjectInfo = Info->Streams[0].ReplicationMap.ReplicatedObjects.Find(TestObjectPath);
+				const FConcertReplicatedObjectInfo* ObjectInfo = Info->Streams[0].ReplicationMap.ReplicatedObjects.Find(TestObjectPath);
 				if (!ObjectInfo)
 				{
 					AddError(TEXT("SkipProperties > No object info"));
@@ -157,7 +157,7 @@ namespace UE::ConcertSyncTests::Replication::Stream
 			{
 				bReceivedResponse_SkipAuthority = true;
 				TestTrue(TEXT("ErrorCode == Handled"), Response.ErrorCode == EReplicationResponseErrorCode::Handled);
-				if (const FReplicationClientQueriedInfo* Info = Response.ClientInfo.Find(SenderEndpointId))
+				if (const FConcertQueriedClientInfo* Info = Response.ClientInfo.Find(SenderEndpointId))
 				{
 					TestEqual(TEXT("SkipAuthority > No authority data"), Info->Authority.Num(), 0);
 				}
@@ -184,7 +184,7 @@ namespace UE::ConcertSyncTests::Replication::Stream
 			{
 				bReceivedResponse_SkipFrequency = true;
 				TestTrue(TEXT("ErrorCode == Handled"), Response.ErrorCode == EReplicationResponseErrorCode::Handled);
-				if (const FReplicationClientQueriedInfo* Info = Response.ClientInfo.Find(SenderEndpointId)
+				if (const FConcertQueriedClientInfo* Info = Response.ClientInfo.Find(SenderEndpointId)
 					; Info && ensureAlways(Info->Streams.Num() == 1))
 				{
 					const FConcertStreamFrequencySettings& FrequencySettings = Info->Streams[0].FrequencySettings;
diff --git a/Engine/Plugins/Developer/Concert/ConcertSync/ConcertSyncTest/Source/ConcertSyncTest/Private/Replication/Util/ChangeStreamsTestBase.cpp b/Engine/Plugins/Developer/Concert/ConcertSync/ConcertSyncTest/Source/ConcertSyncTest/Private/Replication/Util/ChangeStreamsTestBase.cpp
--- a/Engine/Plugins/Developer/Concert/ConcertSync/ConcertSyncTest/Source/ConcertSyncTest/Private/Replication/Util/ChangeStreamsTestBase.cpp
+++ b/Engine/Plugins/Developer/Concert/ConcertSync/ConcertSyncTest/Source/ConcertSyncTest/Private/Replication/Util/ChangeStreamsTestBase.cpp
@@ -2,7 +2,7 @@
 
 #include "ChangeStreamsTestBase.h"
 
-#include "Replication/Data/ReplicationStreamDescription.h"
+#include "Replication/Data/ReplicationStream.h"
 #include "Replication/IConcertClientReplicationManager.h"
 #include "Replication/TestReflectionObject.h"
 #include "Replication/Util/SendReceiveGenericStreamTestBase.h"
@@ -17,15 +17,15 @@ namespace UE::ConcertSyncTests::Replication
 			: FSendReceiveGenericTestBase(InName, bInComplexTask)
 	{}
 	
-	void FChangeStreamsTestBase::AddFloatProperties(const UTestReflectionObject& TestObject, FReplicationStreamDescription& Stream)
+	void FChangeStreamsTestBase::AddFloatProperties(const UTestReflectionObject& TestObject, FConcertReplicationStream& Stream)
 	{
 		FConcertPropertySelection Properties;
 		AddFloatProperty(Properties);
-		const FReplicatedObjectInfo AllProperties { TestObject.GetClass(), MoveTemp(Properties) };
+		const FConcertReplicatedObjectInfo AllProperties { TestObject.GetClass(), MoveTemp(Properties) };
 		Stream.BaseDescription.ReplicationMap.ReplicatedObjects.Add(&TestObject, AllProperties);
 	}
 	
-	TTuple<FGuid, FReplicationStreamDescription> FChangeStreamsTestBase::CreateFloatPropertyStream(
+	TTuple<FGuid, FConcertReplicationStream> FChangeStreamsTestBase::CreateFloatPropertyStream(
 		const UTestReflectionObject& TestObject,
 		FConcertObjectReplicationSettings DefaultFrequencySettings
 		)
@@ -33,14 +33,14 @@ namespace UE::ConcertSyncTests::Replication
 		const FGuid NewStreamId = FGuid::NewGuid();
 
 		// Use Realtime replication because otherwise our test will need to pass latent time due to the frequency system
-		FReplicationStreamDescription SendingStream;
+		FConcertReplicationStream SendingStream;
 		SendingStream.BaseDescription.Identifier = NewStreamId;
 		SendingStream.BaseDescription.FrequencySettings.Defaults = DefaultFrequencySettings;
 		AddFloatProperties(TestObject, SendingStream);
 		return { NewStreamId, SendingStream };
 	}
 
-	TTuple<FGuid, FReplicationStreamDescription> FChangeStreamsTestBase::CreateVectorPropertyStream(
+	TTuple<FGuid, FConcertReplicationStream> FChangeStreamsTestBase::CreateVectorPropertyStream(
 		const UTestReflectionObject& TestObject,
 		FConcertObjectReplicationSettings DefaultFrequencySettings
 		)
@@ -50,15 +50,15 @@ namespace UE::ConcertSyncTests::Replication
 		AddVectorProperty(Properties);
 		
 		// Use Realtime replication because otherwise our test will need to pass latent time due to the frequency system
-		const FReplicatedObjectInfo AllProperties { TestObject.GetClass(), MoveTemp(Properties) };
-		FReplicationStreamDescription SendingStream;
+		const FConcertReplicatedObjectInfo AllProperties { TestObject.GetClass(), MoveTemp(Properties) };
+		FConcertReplicationStream SendingStream;
 		SendingStream.BaseDescription.Identifier = NewStreamId;
 		SendingStream.BaseDescription.ReplicationMap.ReplicatedObjects.Add(&TestObject, AllProperties);
 		SendingStream.BaseDescription.FrequencySettings.Defaults = DefaultFrequencySettings;
 		return { NewStreamId, SendingStream };
 	}
 
-	FConcertPropertySelection& FChangeStreamsTestBase::GetPropertySelection(FReplicationStreamDescription& Stream, const UTestReflectionObject& TestObject)
+	FConcertPropertySelection& FChangeStreamsTestBase::GetPropertySelection(FConcertReplicationStream& Stream, const UTestReflectionObject& TestObject)
 	{
 		return Stream.BaseDescription.ReplicationMap.ReplicatedObjects.FindOrAdd(&TestObject).PropertySelection;
 	}
@@ -80,7 +80,7 @@ namespace UE::ConcertSyncTests::Replication
 	TFuture<FConcertReplication_ChangeStream_Response> FChangeStreamsTestBase::ChangeStreamForSenderClientAndValidate(
 		const FString& InTestName,
 		const FConcertReplication_ChangeStream_Request& Request,
-		const TArray<FSharedReplicationStreamDescription>& ExpectedStreams,
+		const TArray<FConcertBaseStreamInfo>& ExpectedStreams,
 		EChangeStreamValidation ValidationFlags
 		)
 	{
@@ -112,7 +112,7 @@ namespace UE::ConcertSyncTests::Replication
 		return MoveTemp(Future);
 	}
 	
-	void FChangeStreamsTestBase::ValidateSenderClientStreams(const FString& InTestName, const TArray<FSharedReplicationStreamDescription>& ExpectedStreams)
+	void FChangeStreamsTestBase::ValidateSenderClientStreams(const FString& InTestName, const TArray<FConcertBaseStreamInfo>& ExpectedStreams)
 	{
 		using namespace ConcertSyncClient::Replication;
 
@@ -124,15 +124,15 @@ namespace UE::ConcertSyncTests::Replication
 			{
 				bQueriedClientInfo = true;
 				TestTrue(TEXT("ErrorCode == Handled"), Response.ErrorCode == EReplicationResponseErrorCode::Handled);
-				const FReplicationClientQueriedInfo* Info = Response.ClientInfo.Find(SenderId);
+				const FConcertQueriedClientInfo* Info = Response.ClientInfo.Find(SenderId);
 				const bool bRegisteredAndQueriedStreamsAreEqual = Info->Streams == ExpectedStreams;
 				TestTrue(FString::Printf(TEXT("%s > Registered and queried streams are equal"), *InTestName), bRegisteredAndQueriedStreamsAreEqual);
 			});
 
 		// Validate the local client's cache equals  Streams
-		const TArray<FReplicationStreamDescription> LocalStreams = ClientReplicationManager_Sender->GetRegisteredStreams();
-		TArray<FSharedReplicationStreamDescription> TransformedLocalStreams;
-		Algo::Transform(LocalStreams, TransformedLocalStreams, [](const FReplicationStreamDescription& Stream){ return Stream.BaseDescription; });
+		const TArray<FConcertReplicationStream> LocalStreams = ClientReplicationManager_Sender->GetRegisteredStreams();
+		TArray<FConcertBaseStreamInfo> TransformedLocalStreams;
+		Algo::Transform(LocalStreams, TransformedLocalStreams, [](const FConcertReplicationStream& Stream){ return Stream.BaseDescription; });
 		TestEqual(FString::Printf(TEXT("%s > Local streams match registered streams"), *InTestName), TransformedLocalStreams, ExpectedStreams);
 	};
 }
\ No newline at end of file
diff --git a/Engine/Plugins/Developer/Concert/ConcertSync/ConcertSyncTest/Source/ConcertSyncTest/Private/Replication/Util/ChangeStreamsTestBase.h b/Engine/Plugins/Developer/Concert/ConcertSync/ConcertSyncTest/Source/ConcertSyncTest/Private/Replication/Util/ChangeStreamsTestBase.h
--- a/Engine/Plugins/Developer/Concert/ConcertSync/ConcertSyncTest/Source/ConcertSyncTest/Private/Replication/Util/ChangeStreamsTestBase.h
+++ b/Engine/Plugins/Developer/Concert/ConcertSync/ConcertSyncTest/Source/ConcertSyncTest/Private/Replication/Util/ChangeStreamsTestBase.h
@@ -2,7 +2,7 @@
 
 #pragma once
 
-#include "Replication/Data/ReplicationStreamDescription.h"
+#include "Replication/Data/ReplicationStream.h"
 #include "Replication/IConcertClientReplicationManager.h"
 #include "Replication/Messages/ChangeStream.h"
 #include "SendReceiveGenericStreamTestBase.h"
@@ -52,18 +52,18 @@ namespace UE::ConcertSyncTests::Replication
 		virtual ConcertSyncClient::Replication::FJoinReplicatedSessionArgs CreateReceiverArgs() override { return ReceiverArgs; }
 		//~ End FSendReceiveTestBase Interface
 
-		static void AddFloatProperties(const UTestReflectionObject& TestObject, FReplicationStreamDescription& Stream);
+		static void AddFloatProperties(const UTestReflectionObject& TestObject, FConcertReplicationStream& Stream);
 		
-		static TTuple<FGuid, FReplicationStreamDescription> CreateFloatPropertyStream(
+		static TTuple<FGuid, FConcertReplicationStream> CreateFloatPropertyStream(
 			const UTestReflectionObject& TestObject,
 			FConcertObjectReplicationSettings DefaultFrequencySettings = { EConcertObjectReplicationMode::Realtime }
 			);
-		static TTuple<FGuid, FReplicationStreamDescription> CreateVectorPropertyStream(
+		static TTuple<FGuid, FConcertReplicationStream> CreateVectorPropertyStream(
 			const UTestReflectionObject& TestObject,
 			FConcertObjectReplicationSettings DefaultFrequencySettings = { EConcertObjectReplicationMode::Realtime }
 			);
 		
-		static FConcertPropertySelection& GetPropertySelection(FReplicationStreamDescription& Stream, const UTestReflectionObject& TestObject);
+		static FConcertPropertySelection& GetPropertySelection(FConcertReplicationStream& Stream, const UTestReflectionObject& TestObject);
 
 		static void AddFloatProperty(FConcertPropertySelection& PropertySelection);
 		static void AddVectorProperty(FConcertPropertySelection& PropertySelection);
@@ -72,9 +72,9 @@ namespace UE::ConcertSyncTests::Replication
 		TFuture<FConcertReplication_ChangeStream_Response> ChangeStreamForSenderClientAndValidate(
 			const FString& InTestName,
 			const FConcertReplication_ChangeStream_Request& Request,
-			const TArray<FSharedReplicationStreamDescription>& ExpectedStreams,
+			const TArray<FConcertBaseStreamInfo>& ExpectedStreams,
 			EChangeStreamValidation ValidationFlags = EChangeStreamValidation::TestWasHandled | EChangeStreamValidation::TestResponseSuccess
 			);
-		void ValidateSenderClientStreams(const FString& InTestName, const TArray<FSharedReplicationStreamDescription>& ExpectedStreams);
+		void ValidateSenderClientStreams(const FString& InTestName, const TArray<FConcertBaseStreamInfo>& ExpectedStreams);
 	};
 }
diff --git a/Engine/Plugins/Developer/Concert/ConcertSync/ConcertSyncTest/Source/ConcertSyncTest/Private/Replication/Util/SendReceiveObjectTestBase.cpp b/Engine/Plugins/Developer/Concert/ConcertSync/ConcertSyncTest/Source/ConcertSyncTest/Private/Replication/Util/SendReceiveObjectTestBase.cpp
--- a/Engine/Plugins/Developer/Concert/ConcertSync/ConcertSyncTest/Source/ConcertSyncTest/Private/Replication/Util/SendReceiveObjectTestBase.cpp
+++ b/Engine/Plugins/Developer/Concert/ConcertSync/ConcertSyncTest/Source/ConcertSyncTest/Private/Replication/Util/SendReceiveObjectTestBase.cpp
@@ -3,13 +3,8 @@
 #include "SendReceiveObjectTestBase.h"
 
 #include "ConcertClientReplicationBridgeMock.h"
-#include "Replication/Data/ReplicationStreamDescription.h"
 #include "Replication/IConcertClientReplicationManager.h"
-#include "Replication/IConcertServerReplicationManager.h"
 #include "Replication/Messages/ObjectReplication.h"
-#include "Replication/Messages/Handshake.h"
-#include "Replication/PropertyChainUtils.h"
-#include "Replication/ReplicationTestInterface.h"
 #include "Replication/TestReflectionObject.h"
 #include "Util/ClientServerCommunicationTest.h"
 
diff --git a/Engine/Plugins/Developer/Concert/ConcertSync/ConcertSyncTest/Source/ConcertSyncTest/Private/Replication/Util/SendReceiveTestBase.cpp b/Engine/Plugins/Developer/Concert/ConcertSync/ConcertSyncTest/Source/ConcertSyncTest/Private/Replication/Util/SendReceiveTestBase.cpp
--- a/Engine/Plugins/Developer/Concert/ConcertSync/ConcertSyncTest/Source/ConcertSyncTest/Private/Replication/Util/SendReceiveTestBase.cpp
+++ b/Engine/Plugins/Developer/Concert/ConcertSync/ConcertSyncTest/Source/ConcertSyncTest/Private/Replication/Util/SendReceiveTestBase.cpp
@@ -31,14 +31,14 @@ namespace UE::ConcertSyncTests::Replication
 	{
 		ConcertSyncClient::Replication::FJoinReplicatedSessionArgs SenderJoinArgs;
 		
-		FReplicatedObjectInfo ReplicatedObjectInfo { Object.GetClass() };
+		FConcertReplicatedObjectInfo ReplicatedObjectInfo { Object.GetClass() };
 		ConcertSyncCore::PropertyChain::ForEachReplicatableConcertProperty(*Object.GetClass(), [&ReplicatedObjectInfo](FConcertPropertyChain&& Chain)
 		{
 			ReplicatedObjectInfo.PropertySelection.ReplicatedProperties.Emplace(MoveTemp(Chain));
 			return EBreakBehavior::Continue;
 		});
 
-		FReplicationStreamDescription SendingStream;
+		FConcertReplicationStream SendingStream;
 		SendingStream.BaseDescription.Identifier = SenderStreamId;
 		SendingStream.BaseDescription.ReplicationMap.ReplicatedObjects.Add(&Object, ReplicatedObjectInfo);
 		SendingStream.BaseDescription.FrequencySettings.Defaults = { ReplicationMode, ReplicationRate };
diff --git a/Engine/Plugins/Developer/Concert/ConcertSync/ConcertSyncTest/Source/ConcertSyncTest/Private/Replication/Util/SendReceiveTestBase.h b/Engine/Plugins/Developer/Concert/ConcertSync/ConcertSyncTest/Source/ConcertSyncTest/Private/Replication/Util/SendReceiveTestBase.h
--- a/Engine/Plugins/Developer/Concert/ConcertSync/ConcertSyncTest/Source/ConcertSyncTest/Private/Replication/Util/SendReceiveTestBase.h
+++ b/Engine/Plugins/Developer/Concert/ConcertSync/ConcertSyncTest/Source/ConcertSyncTest/Private/Replication/Util/SendReceiveTestBase.h
@@ -33,7 +33,7 @@ namespace UE::ConcertSyncTests::Replication
 		 * Creates joining arguments containing a stream that will replicate a single object.
 		 * @param Object The object to replicate
 		 * @param SenderStreamId The ID of the new stream
-		 * @param ReplicationMode The replication mode for the object. Set to Realtime by default so tests need not handle latency (due to FReplicatedObjectInfo::ReplicationRate).
+		 * @param ReplicationMode The replication mode for the object. Set to Realtime by default so tests need not handle latency (due to FConcertReplicatedObjectInfo::ReplicationRate).
 		 * @param ReplicationRate The replication rate for the object. Only used if ReplicationMode == EConcertObjectReplicationMode::SpecifiedRate
 		 */
 		static ConcertSyncClient::Replication::FJoinReplicatedSessionArgs CreateHandshakeArgsFrom(
diff --git a/Engine/Plugins/Developer/Concert/ConcertUI/ConcertClientSharedSlate/Source/ConcertClientSharedSlate/Private/Replication/ClientReplicationWidgetFactories.cpp b/Engine/Plugins/Developer/Concert/ConcertUI/ConcertClientSharedSlate/Source/ConcertClientSharedSlate/Private/Replication/ClientReplicationWidgetFactories.cpp
--- a/Engine/Plugins/Developer/Concert/ConcertUI/ConcertClientSharedSlate/Source/ConcertClientSharedSlate/Private/Replication/ClientReplicationWidgetFactories.cpp
+++ b/Engine/Plugins/Developer/Concert/ConcertUI/ConcertClientSharedSlate/Source/ConcertClientSharedSlate/Private/Replication/ClientReplicationWidgetFactories.cpp
@@ -34,7 +34,7 @@ namespace UE::ConcertClientSharedSlate
 		const EObjectFlags Flags = RF_Transient | RF_Transactional;
 		UReplicationStreamObject* Object = NewObject<UReplicationStreamObject>(GetTransientPackage(), NAME_None, Flags);
 
-		TAttribute<FObjectReplicationMap*> Attribute = TAttribute<FObjectReplicationMap*>::CreateLambda([WeakPtr = TWeakObjectPtr<UReplicationStreamObject>(Object)]() -> FObjectReplicationMap* 
+		TAttribute<FConcertObjectReplicationMap*> Attribute = TAttribute<FConcertObjectReplicationMap*>::CreateLambda([WeakPtr = TWeakObjectPtr<UReplicationStreamObject>(Object)]() -> FConcertObjectReplicationMap* 
 		{
 			if (WeakPtr.IsValid())
 			{
diff --git a/Engine/Plugins/Developer/Concert/ConcertUI/ConcertClientSharedSlate/Source/ConcertClientSharedSlate/Private/Replication/Editor/Model/ReplicationStreamObject.h b/Engine/Plugins/Developer/Concert/ConcertUI/ConcertClientSharedSlate/Source/ConcertClientSharedSlate/Private/Replication/Editor/Model/ReplicationStreamObject.h
--- a/Engine/Plugins/Developer/Concert/ConcertUI/ConcertClientSharedSlate/Source/ConcertClientSharedSlate/Private/Replication/Editor/Model/ReplicationStreamObject.h
+++ b/Engine/Plugins/Developer/Concert/ConcertUI/ConcertClientSharedSlate/Source/ConcertClientSharedSlate/Private/Replication/Editor/Model/ReplicationStreamObject.h
@@ -6,7 +6,7 @@
 #include "UObject/Object.h"
 #include "ReplicationStreamObject.generated.h"
 
-/** UObject wrapper to allow transactions on FObjectReplicationMap. */
+/** UObject wrapper to allow transactions on FConcertObjectReplicationMap. */
 UCLASS()
 class UReplicationStreamObject : public UObject
 {
@@ -14,5 +14,5 @@ class UReplicationStreamObject : public UObject
 public:
 	
 	UPROPERTY()
-	FObjectReplicationMap ReplicationMap;
+	FConcertObjectReplicationMap ReplicationMap;
 };
\ No newline at end of file
diff --git a/Engine/Plugins/Developer/Concert/ConcertUI/ConcertClientSharedSlate/Source/ConcertClientSharedSlate/Private/Replication/Editor/Model/TransactionalReplicationStreamModel.h b/Engine/Plugins/Developer/Concert/ConcertUI/ConcertClientSharedSlate/Source/ConcertClientSharedSlate/Private/Replication/Editor/Model/TransactionalReplicationStreamModel.h
--- a/Engine/Plugins/Developer/Concert/ConcertUI/ConcertClientSharedSlate/Source/ConcertClientSharedSlate/Private/Replication/Editor/Model/TransactionalReplicationStreamModel.h
+++ b/Engine/Plugins/Developer/Concert/ConcertUI/ConcertClientSharedSlate/Source/ConcertClientSharedSlate/Private/Replication/Editor/Model/TransactionalReplicationStreamModel.h
@@ -6,13 +6,13 @@
 #include "Replication/Editor/Model/IEditableReplicationStreamModel.h"
 #include "UObject/GCObject.h"
 
-struct FObjectReplicationMap;
+struct FConcertObjectReplicationMap;
 
 namespace UE::ConcertSharedSlate
 {
 	class IStreamExtender;
 
-	/** Special case of FGenericPropertySelectionModel where the edited FObjectReplicationMap lives in an UObject that is RF_Transactional. */
+	/** Special case of FGenericPropertySelectionModel where the edited FConcertObjectReplicationMap lives in an UObject that is RF_Transactional. */
 	class FTransactionalReplicationStreamModel
 		: public IEditableReplicationStreamModel
 		, public FSelfRegisteringEditorUndoClient
diff --git a/Engine/Plugins/Developer/Concert/ConcertUI/ConcertClientSharedSlate/Source/ConcertClientSharedSlate/Public/Replication/ClientReplicationWidgetFactories.h b/Engine/Plugins/Developer/Concert/ConcertUI/ConcertClientSharedSlate/Source/ConcertClientSharedSlate/Public/Replication/ClientReplicationWidgetFactories.h
--- a/Engine/Plugins/Developer/Concert/ConcertUI/ConcertClientSharedSlate/Source/ConcertClientSharedSlate/Public/Replication/ClientReplicationWidgetFactories.h
+++ b/Engine/Plugins/Developer/Concert/ConcertUI/ConcertClientSharedSlate/Source/ConcertClientSharedSlate/Public/Replication/ClientReplicationWidgetFactories.h
@@ -5,7 +5,7 @@
 #include "Templates/SharedPointer.h"
 
 struct FConcertStreamObjectAutoBindingRules;
-struct FObjectReplicationMap;
+struct FConcertObjectReplicationMap;
 
 namespace UE::ConcertSharedSlate
 {
@@ -25,7 +25,7 @@ namespace UE::ConcertClientSharedSlate
 	 * Wraps the passed in BaseModel and makes it transactional.
 	 * All calls that modify the underlying model was wrapped with scoped transactions.
 	 * 
-	 * @param OwnerObject The object containing the FObjectReplicationMap - used for transactions.
+	 * @param OwnerObject The object containing the FConcertObjectReplicationMap - used for transactions.
 	 */
 	CONCERTCLIENTSHAREDSLATE_API TSharedRef<ConcertSharedSlate::IEditableReplicationStreamModel> CreateTransactionalStreamModel(
 		TSharedRef<ConcertSharedSlate::IEditableReplicationStreamModel> BaseModel,
diff --git a/Engine/Plugins/Developer/Concert/ConcertUI/ConcertClientSharedSlate/Source/ConcertClientSharedSlate/Public/Replication/Settings/ConcertStreamObjectAutoBindingRules.h b/Engine/Plugins/Developer/Concert/ConcertUI/ConcertClientSharedSlate/Source/ConcertClientSharedSlate/Public/Replication/Settings/ConcertStreamObjectAutoBindingRules.h
--- a/Engine/Plugins/Developer/Concert/ConcertUI/ConcertClientSharedSlate/Source/ConcertClientSharedSlate/Public/Replication/Settings/ConcertStreamObjectAutoBindingRules.h
+++ b/Engine/Plugins/Developer/Concert/ConcertUI/ConcertClientSharedSlate/Source/ConcertClientSharedSlate/Public/Replication/Settings/ConcertStreamObjectAutoBindingRules.h
@@ -9,7 +9,7 @@
 
 class UObject;
 struct FConcertPropertyChain;
-struct FReplicatedObjectInfo;
+struct FConcertReplicatedObjectInfo;
 
 USTRUCT()
 struct CONCERTCLIENTSHAREDSLATE_API FConcertDefaultPropertySelection : public FConcertInheritableClassOption
diff --git a/Engine/Plugins/Developer/Concert/ConcertUI/ConcertSharedSlate/Source/ConcertSharedSlate/Private/Replication/Editor/Model/GenericReplicationStreamModel.cpp b/Engine/Plugins/Developer/Concert/ConcertUI/ConcertSharedSlate/Source/ConcertSharedSlate/Private/Replication/Editor/Model/GenericReplicationStreamModel.cpp
--- a/Engine/Plugins/Developer/Concert/ConcertUI/ConcertSharedSlate/Source/ConcertSharedSlate/Private/Replication/Editor/Model/GenericReplicationStreamModel.cpp
+++ b/Engine/Plugins/Developer/Concert/ConcertUI/ConcertSharedSlate/Source/ConcertSharedSlate/Private/Replication/Editor/Model/GenericReplicationStreamModel.cpp
@@ -32,7 +32,7 @@ namespace UE::ConcertSharedSlate
 	}
 
 	FGenericReplicationStreamModel::FGenericReplicationStreamModel(
-		TAttribute<FObjectReplicationMap*> InReplicationMapAttribute,
+		TAttribute<FConcertObjectReplicationMap*> InReplicationMapAttribute,
 		TSharedPtr<IStreamExtender> InExtender
 		)
 		: ReplicationMapAttribute(MoveTemp(InReplicationMapAttribute))
@@ -41,13 +41,13 @@ namespace UE::ConcertSharedSlate
 
 	FSoftClassPath FGenericReplicationStreamModel::GetObjectClass(const FSoftObjectPath& Object) const
 	{
-		const FObjectReplicationMap* ReplicationMap = ReplicationMapAttribute.Get();
+		const FConcertObjectReplicationMap* ReplicationMap = ReplicationMapAttribute.Get();
 		if (!ensure(ReplicationMap))
 		{
 			return {};
 		}
 		
-		const FReplicatedObjectInfo* AssignedProperties = ReplicationMap->ReplicatedObjects.Find(Object);
+		const FConcertReplicatedObjectInfo* AssignedProperties = ReplicationMap->ReplicatedObjects.Find(Object);
 		return AssignedProperties
 			? AssignedProperties->ClassPath
 			: FSoftClassPath{};
@@ -55,33 +55,33 @@ namespace UE::ConcertSharedSlate
 
 	bool FGenericReplicationStreamModel::ContainsObjects(const TSet<FSoftObjectPath>& Objects) const
 	{
-		const FObjectReplicationMap* ReplicationMap = ReplicationMapAttribute.Get();
+		const FConcertObjectReplicationMap* ReplicationMap = ReplicationMapAttribute.Get();
 		return ensure(ReplicationMap)
 			&& Algo::AllOf(Objects, [this, ReplicationMap](const FSoftObjectPath& ObjectPath){ return ReplicationMap->ReplicatedObjects.Contains(ObjectPath); });
 	}
 
 	bool FGenericReplicationStreamModel::ContainsProperties(const FSoftObjectPath& Object, const TSet<FConcertPropertyChain>& Properties) const
 	{
-		const FObjectReplicationMap* ReplicationMap = ReplicationMapAttribute.Get();
+		const FConcertObjectReplicationMap* ReplicationMap = ReplicationMapAttribute.Get();
 		if (!ensure(ReplicationMap))
 		{
 			return false;
 		}
 
-		const FReplicatedObjectInfo* ObjectInfo = ReplicationMap->ReplicatedObjects.Find(Object);
+		const FConcertReplicatedObjectInfo* ObjectInfo = ReplicationMap->ReplicatedObjects.Find(Object);
 		return ObjectInfo
 			&& Algo::AllOf(Properties, [ObjectInfo](const FConcertPropertyChain& Property){ return ObjectInfo->PropertySelection.ReplicatedProperties.Contains(Property); });
 	}
 
 	bool FGenericReplicationStreamModel::ForEachReplicatedObject(TFunctionRef<EBreakBehavior(const FSoftObjectPath& Object)> Delegate) const
 	{
-		const FObjectReplicationMap* ReplicationMap = ReplicationMapAttribute.Get();
+		const FConcertObjectReplicationMap* ReplicationMap = ReplicationMapAttribute.Get();
 		if (!ensure(ReplicationMap))
 		{
 			return false;
 		}
 
-		for (const TPair<FSoftObjectPath, FReplicatedObjectInfo>& ObjectMap: ReplicationMap->ReplicatedObjects)
+		for (const TPair<FSoftObjectPath, FConcertReplicatedObjectInfo>& ObjectMap: ReplicationMap->ReplicatedObjects)
 		{
 			if (Delegate(ObjectMap.Key) == EBreakBehavior::Break)
 			{
@@ -94,13 +94,13 @@ namespace UE::ConcertSharedSlate
 
 	bool FGenericReplicationStreamModel::ForEachProperty(const FSoftObjectPath& Object, TFunctionRef<EBreakBehavior(const FConcertPropertyChain& Parent)> Delegate) const
 	{
-		const FObjectReplicationMap* ReplicationMap = ReplicationMapAttribute.Get();
+		const FConcertObjectReplicationMap* ReplicationMap = ReplicationMapAttribute.Get();
 		if (!ensure(ReplicationMap))
 		{
 			return false;
 		}
 
-		const FReplicatedObjectInfo* AssignedProperties = ReplicationMap->ReplicatedObjects.Find(Object);
+		const FConcertReplicatedObjectInfo* AssignedProperties = ReplicationMap->ReplicatedObjects.Find(Object);
 		if (!AssignedProperties)
 		{
 			return false;
@@ -118,7 +118,7 @@ namespace UE::ConcertSharedSlate
 
 	void FGenericReplicationStreamModel::AddObjects(TConstArrayView<UObject*> Objects)
 	{
-		FObjectReplicationMap* ReplicationMap = ReplicationMapAttribute.Get();
+		FConcertObjectReplicationMap* ReplicationMap = ReplicationMapAttribute.Get();
 		if (!ensure(ReplicationMap) || Objects.IsEmpty())
 		{
 			return;
@@ -130,7 +130,7 @@ namespace UE::ConcertSharedSlate
 			const FSoftObjectPath ObjectPath = Object;
 			if (ensureAlways(Object) && !ReplicationMap->ReplicatedObjects.Contains(ObjectPath))
 			{
-				FReplicatedObjectInfo& ObjectInfo = ReplicationMap->ReplicatedObjects.Add(ObjectPath);
+				FConcertReplicatedObjectInfo& ObjectInfo = ReplicationMap->ReplicatedObjects.Add(ObjectPath);
 				ObjectInfo.ClassPath = Object->GetClass();
 				AddedObjects.AddUnique(Object);
 				
@@ -146,7 +146,7 @@ namespace UE::ConcertSharedSlate
 
 	void FGenericReplicationStreamModel::RemoveObjects(TConstArrayView<FSoftObjectPath> Objects)
 	{
-		FObjectReplicationMap* ReplicationMap = ReplicationMapAttribute.Get();
+		FConcertObjectReplicationMap* ReplicationMap = ReplicationMapAttribute.Get();
 		if (!ensure(ReplicationMap) || Objects.IsEmpty())
 		{
 			return;
@@ -183,13 +183,13 @@ namespace UE::ConcertSharedSlate
 
 	void FGenericReplicationStreamModel::AddProperties(const FSoftObjectPath& Object, TConstArrayView<FConcertPropertyChain> Properties)
 	{
-		FObjectReplicationMap* ReplicationMap = ReplicationMapAttribute.Get();
+		FConcertObjectReplicationMap* ReplicationMap = ReplicationMapAttribute.Get();
 		if (!ensure(ReplicationMap))
 		{
 			return;
 		}
 		
-		FReplicatedObjectInfo* AssignedProperties = ReplicationMap->ReplicatedObjects.Find(Object);
+		FConcertReplicatedObjectInfo* AssignedProperties = ReplicationMap->ReplicatedObjects.Find(Object);
 		if (!AssignedProperties)
 		{
 			return;
@@ -222,13 +222,13 @@ namespace UE::ConcertSharedSlate
 
 	void FGenericReplicationStreamModel::RemoveProperties(const FSoftObjectPath& Object, TConstArrayView<FConcertPropertyChain> Properties)
 	{
-		FObjectReplicationMap* ReplicationMap = ReplicationMapAttribute.Get();
+		FConcertObjectReplicationMap* ReplicationMap = ReplicationMapAttribute.Get();
 		if (!ensure(ReplicationMap))
 		{
 			return;
 		}
 		
-		FReplicatedObjectInfo* AssignedProperties = ReplicationMap->ReplicatedObjects.Find(Object);
+		FConcertReplicatedObjectInfo* AssignedProperties = ReplicationMap->ReplicatedObjects.Find(Object);
 		if (!AssignedProperties)
 		{
 			return;
@@ -264,7 +264,7 @@ namespace UE::ConcertSharedSlate
 		}
 	}
 
-	void FGenericReplicationStreamModel::ExtendObjects(FObjectReplicationMap& ReplicationMap, UObject& AddedObject, TArray<UObject*>& ObjectsAddedSoFar)
+	void FGenericReplicationStreamModel::ExtendObjects(FConcertObjectReplicationMap& ReplicationMap, UObject& AddedObject, TArray<UObject*>& ObjectsAddedSoFar)
 	{
 		if (!Extender)
 		{
@@ -276,10 +276,10 @@ namespace UE::ConcertSharedSlate
 		public:
 
 			TQueue<UObject*> ObjectsToProcess;
-			FObjectReplicationMap& ReplicationMap;
+			FConcertObjectReplicationMap& ReplicationMap;
 			TArray<UObject*>& ObjectsAddedSoFar;
 
-			explicit FExtensionContext(FObjectReplicationMap& ReplicationMap, TArray<UObject*>& ObjectsAddedSoFar)
+			explicit FExtensionContext(FConcertObjectReplicationMap& ReplicationMap, TArray<UObject*>& ObjectsAddedSoFar)
 				: ReplicationMap(ReplicationMap)
 				, ObjectsAddedSoFar(ObjectsAddedSoFar)
 			{}
@@ -293,7 +293,7 @@ namespace UE::ConcertSharedSlate
 				}
 				
 				AddAdditionalObject(Object);
-				FReplicatedObjectInfo& ObjectInfo = ReplicationMap.ReplicatedObjects[&Object];
+				FConcertReplicatedObjectInfo& ObjectInfo = ReplicationMap.ReplicatedObjects[&Object];
 
 				constexpr bool bLog = false;
 				const FProperty* ResolvedProperty = PropertyChain.ResolveProperty(*ObjectClass, bLog);
diff --git a/Engine/Plugins/Developer/Concert/ConcertUI/ConcertSharedSlate/Source/ConcertSharedSlate/Private/Replication/Editor/Model/GenericReplicationStreamModel.h b/Engine/Plugins/Developer/Concert/ConcertUI/ConcertSharedSlate/Source/ConcertSharedSlate/Private/Replication/Editor/Model/GenericReplicationStreamModel.h
--- a/Engine/Plugins/Developer/Concert/ConcertUI/ConcertSharedSlate/Source/ConcertSharedSlate/Private/Replication/Editor/Model/GenericReplicationStreamModel.h
+++ b/Engine/Plugins/Developer/Concert/ConcertUI/ConcertSharedSlate/Source/ConcertSharedSlate/Private/Replication/Editor/Model/GenericReplicationStreamModel.h
@@ -7,20 +7,20 @@
 #include "UObject/WeakObjectPtrTemplates.h"
 
 struct FConcertStreamObjectAutoBindingRules;
-struct FObjectReplicationMap;
+struct FConcertObjectReplicationMap;
 
 namespace UE::ConcertSharedSlate
 {
 	class IStreamExtender;
 
-	/** Implements logic for editing a FObjectReplicationMap contained in an UObject. */
+	/** Implements logic for editing a FConcertObjectReplicationMap contained in an UObject. */
 	class FGenericReplicationStreamModel
 		: public IEditableReplicationStreamModel
 	{
 	public:
 		
 		FGenericReplicationStreamModel(
-			TAttribute<FObjectReplicationMap*> InReplicationMapAttribute,
+			TAttribute<FConcertObjectReplicationMap*> InReplicationMapAttribute,
 			TSharedPtr<IStreamExtender> InExtender = nullptr
 			);
 		
@@ -44,7 +44,7 @@ namespace UE::ConcertSharedSlate
 	private:
 
 		/** Returns the replication map that is supposed to be edited. */
-		TAttribute<FObjectReplicationMap*> ReplicationMapAttribute;
+		TAttribute<FConcertObjectReplicationMap*> ReplicationMapAttribute;
 
 		/** Adds properties and objects when an object is added. Can be null. */
 		TSharedPtr<IStreamExtender> Extender;
@@ -53,7 +53,7 @@ namespace UE::ConcertSharedSlate
 		FOnPropertiesChanged OnPropertiesChangedDelegate;
 
 		/** Applies Extender to AddedObject whilst adding any additionally added UObjects to ObjectsAddedSoFar and avoiding adding objects already added to ObjectsAddedSoFar. */
-		void ExtendObjects(FObjectReplicationMap& ReplicationMap, UObject& AddedObject, TArray<UObject*>& ObjectsAddedSoFar);
+		void ExtendObjects(FConcertObjectReplicationMap& ReplicationMap, UObject& AddedObject, TArray<UObject*>& ObjectsAddedSoFar);
 	};
 }
 
diff --git a/Engine/Plugins/Developer/Concert/ConcertUI/ConcertSharedSlate/Source/ConcertSharedSlate/Private/Replication/ReplicationWidgetFactories.cpp b/Engine/Plugins/Developer/Concert/ConcertUI/ConcertSharedSlate/Source/ConcertSharedSlate/Private/Replication/ReplicationWidgetFactories.cpp
--- a/Engine/Plugins/Developer/Concert/ConcertUI/ConcertSharedSlate/Source/ConcertSharedSlate/Private/Replication/ReplicationWidgetFactories.cpp
+++ b/Engine/Plugins/Developer/Concert/ConcertUI/ConcertSharedSlate/Source/ConcertSharedSlate/Private/Replication/ReplicationWidgetFactories.cpp
@@ -9,7 +9,7 @@
 namespace UE::ConcertSharedSlate
 {
 	TSharedRef<IEditableReplicationStreamModel> CreateBaseStreamModel(
-		TAttribute<FObjectReplicationMap*> ReplicationMapAttribute,
+		TAttribute<FConcertObjectReplicationMap*> ReplicationMapAttribute,
 		TSharedPtr<IStreamExtender> Extender
 		)
 	{
diff --git a/Engine/Plugins/Developer/Concert/ConcertUI/ConcertSharedSlate/Source/ConcertSharedSlate/Public/Replication/ReplicationWidgetFactories.h b/Engine/Plugins/Developer/Concert/ConcertUI/ConcertSharedSlate/Source/ConcertSharedSlate/Public/Replication/ReplicationWidgetFactories.h
--- a/Engine/Plugins/Developer/Concert/ConcertUI/ConcertSharedSlate/Source/ConcertSharedSlate/Public/Replication/ReplicationWidgetFactories.h
+++ b/Engine/Plugins/Developer/Concert/ConcertUI/ConcertSharedSlate/Source/ConcertSharedSlate/Public/Replication/ReplicationWidgetFactories.h
@@ -14,7 +14,7 @@ class UObject;
 class SWidget;
 
 struct FConcertStreamObjectAutoBindingRules;
-struct FObjectReplicationMap;
+struct FConcertObjectReplicationMap;
 
 namespace UE::ConcertSharedSlate
 {
@@ -34,13 +34,13 @@ namespace UE::ConcertSharedSlate
 	/**
 	 * Creates a model that can be passed to CreateEditor.
 	 * 
-	 * @param ReplicationMapAttribute Getter for extracting the FObjectReplicationMap to edit
+	 * @param ReplicationMapAttribute Getter for extracting the FConcertObjectReplicationMap to edit
 	 * @param Extender Optional callbacks for adding additional properties and objects when an object is added to the model
 	 * 
-	 * @return A model that will edit the FObjectReplicationMap.
+	 * @return A model that will edit the FConcertObjectReplicationMap.
 	 */
 	CONCERTSHAREDSLATE_API TSharedRef<IEditableReplicationStreamModel> CreateBaseStreamModel(
-		TAttribute<FObjectReplicationMap*> ReplicationMapAttribute,
+		TAttribute<FConcertObjectReplicationMap*> ReplicationMapAttribute,
 		TSharedPtr<IStreamExtender> Extender = nullptr
 		);
 	
