Fix LiveEdit being broken since the new property bag checkin. Was related to how we store property types now. Fixed by making the archive persistent.

diff --git a/Engine/Plugins/Developer/Concert/ConcertSync/ConcertSyncCore/Source/ConcertSyncCore/Private/ConcertSyncArchives.cpp b/Engine/Plugins/Developer/Concert/ConcertSync/ConcertSyncCore/Source/ConcertSyncCore/Private/ConcertSyncArchives.cpp
--- a/Engine/Plugins/Developer/Concert/ConcertSync/ConcertSyncCore/Source/ConcertSyncCore/Private/ConcertSyncArchives.cpp
+++ b/Engine/Plugins/Developer/Concert/ConcertSync/ConcertSyncCore/Source/ConcertSyncCore/Private/ConcertSyncArchives.cpp
@@ -228,7 +228,7 @@ bool FConcertSyncWorldRemapper::HasMapping() const
 }
 
 FConcertSyncObjectWriter::FConcertSyncObjectWriter(FConcertLocalIdentifierTable* InLocalIdentifierTable, UObject* InObj, TArray<uint8>& OutBytes, const bool InIncludeEditorOnlyData, const bool InSkipAssets, const FConcertSyncRemapObjectPath& InRemapDelegate)
-	: FConcertIdentifierWriter(InLocalIdentifierTable, OutBytes, /*bIsPersistent*/false)
+	: FConcertIdentifierWriter(InLocalIdentifierTable, OutBytes, /*bIsPersistent*/true)
 	, bSkipAssets(InSkipAssets)
 	, ShouldSkipPropertyFunc()
 	, RemapObjectPathDelegate(InRemapDelegate)
@@ -446,7 +446,7 @@ void InitReaderArchive(FArchive& Ar, const FConcertSessionVersionInfo* InVersion
 } // namespace UE::Concert::Private::ConcertSyncArchiveUtil
 
 FConcertSyncObjectReader::FConcertSyncObjectReader(const FConcertLocalIdentifierTable* InLocalIdentifierTable, FConcertSyncWorldRemapper InWorldRemapper, const FConcertSessionVersionInfo* InVersionInfo, UObject* InObj, const TArray<uint8>& InBytes, const FConcertSyncEncounteredMissingObject& InEncounteredMissingObjectDelegate)
-	: FConcertIdentifierReader(InLocalIdentifierTable, InBytes, /*bIsPersistent*/false)
+	: FConcertIdentifierReader(InLocalIdentifierTable, InBytes, /*bIsPersistent*/true)
 	, WorldRemapper(MoveTemp(InWorldRemapper))
 	, EncounteredMissingObjectDelegate(InEncounteredMissingObjectDelegate)
 {
@@ -614,7 +614,7 @@ FString FConcertSyncObjectReader::GetArchiveName() const
 
 
 FConcertSyncObjectRewriter::FConcertSyncObjectRewriter(const FConcertLocalIdentifierTable* InLocalIdentifierTable, FConcertLocalIdentifierTable* InRewriteIdentifierTable, const FConcertSessionVersionInfo* InVersionInfo, TArray<uint8>& InBytes)
-	: FConcertIdentifierRewriter(InLocalIdentifierTable, InRewriteIdentifierTable, InBytes, /*bIsPersistent*/false)
+	: FConcertIdentifierRewriter(InLocalIdentifierTable, InRewriteIdentifierTable, InBytes, /*bIsPersistent*/true)
 {
 	UE::Concert::Private::ConcertSyncArchiveUtil::InitReaderArchive(*this, InVersionInfo);
 }
