Refactor FConcertPropertySelection::ReplicatedProperties to be TSet instead of TArray

diff --git a/Engine/Plugins/Developer/Concert/ConcertApp/MultiUserClient/Source/MultiUserClient/Private/Replication/Submission/External/ClientChangeConversionUtils.cpp b/Engine/Plugins/Developer/Concert/ConcertApp/MultiUserClient/Source/MultiUserClient/Private/Replication/Submission/External/ClientChangeConversionUtils.cpp
--- a/Engine/Plugins/Developer/Concert/ConcertApp/MultiUserClient/Source/MultiUserClient/Private/Replication/Submission/External/ClientChangeConversionUtils.cpp
+++ b/Engine/Plugins/Developer/Concert/ConcertApp/MultiUserClient/Source/MultiUserClient/Private/Replication/Submission/External/ClientChangeConversionUtils.cpp
@@ -25,7 +25,7 @@ namespace UE::MultiUserClient::ClientChangeConversionUtils
 			if (ExistingObjectInfo)
 			{
 				FConcertReplicatedObjectInfo NewObjectInfo = FConcertReplicatedObjectInfo::Make(Object);
-				TArray<FConcertPropertyChain>& ReplicatedProperties = NewObjectInfo.PropertySelection.ReplicatedProperties;
+				TSet<FConcertPropertyChain>& ReplicatedProperties = NewObjectInfo.PropertySelection.ReplicatedProperties;
 				switch (PropertyChange.ChangeType)
 				{
 				case EPropertyChangeType::Put:
@@ -35,10 +35,10 @@ namespace UE::MultiUserClient::ClientChangeConversionUtils
 					ReplicatedProperties.Append(PropertyChange.Properties);
 					break;
 				case EPropertyChangeType::Remove:
-					ReplicatedProperties.SetNum(Algo::RemoveIf(ReplicatedProperties, [&PropertyChange](const FConcertPropertyChain& Property)
+					for (const FConcertPropertyChain& Property : PropertyChange.Properties)
 					{
-						return PropertyChange.Properties.Contains(Property);
-					}));
+						ReplicatedProperties.Remove(Property);
+					}
 					break;
 				default: checkNoEntry(); break;
 				}
@@ -47,7 +47,7 @@ namespace UE::MultiUserClient::ClientChangeConversionUtils
 			else
 			{
 				FConcertReplicatedObjectInfo NewObjectInfo = FConcertReplicatedObjectInfo::Make(Object);
-				TArray<FConcertPropertyChain>& ReplicatedProperties = NewObjectInfo.PropertySelection.ReplicatedProperties;
+				TSet<FConcertPropertyChain>& ReplicatedProperties = NewObjectInfo.PropertySelection.ReplicatedProperties;
 				switch (PropertyChange.ChangeType)
 				{
 				case EPropertyChangeType::Put:
diff --git a/Engine/Plugins/Developer/Concert/ConcertApp/MultiUserClient/Source/MultiUserClient/Private/Replication/Submission/MultiEdit/ReassignObjectPropertiesLogic.cpp b/Engine/Plugins/Developer/Concert/ConcertApp/MultiUserClient/Source/MultiUserClient/Private/Replication/Submission/MultiEdit/ReassignObjectPropertiesLogic.cpp
--- a/Engine/Plugins/Developer/Concert/ConcertApp/MultiUserClient/Source/MultiUserClient/Private/Replication/Submission/MultiEdit/ReassignObjectPropertiesLogic.cpp
+++ b/Engine/Plugins/Developer/Concert/ConcertApp/MultiUserClient/Source/MultiUserClient/Private/Replication/Submission/MultiEdit/ReassignObjectPropertiesLogic.cpp
@@ -153,7 +153,7 @@ namespace UE::MultiUserClient
 					{
 						for (const FConcertPropertyChain& PropertyChain : CurrentInfo->PropertySelection.ReplicatedProperties)
 						{
-							PutRequest->Properties.ReplicatedProperties.AddUnique(PropertyChain);
+							PutRequest->Properties.ReplicatedProperties.Add(PropertyChain);
 						}
 					}
 					
diff --git a/Engine/Plugins/Developer/Concert/ConcertApp/MultiUserClient/Source/MultiUserClient/Public/Replication/ChangeOperationTypes.h b/Engine/Plugins/Developer/Concert/ConcertApp/MultiUserClient/Source/MultiUserClient/Public/Replication/ChangeOperationTypes.h
--- a/Engine/Plugins/Developer/Concert/ConcertApp/MultiUserClient/Source/MultiUserClient/Public/Replication/ChangeOperationTypes.h
+++ b/Engine/Plugins/Developer/Concert/ConcertApp/MultiUserClient/Source/MultiUserClient/Public/Replication/ChangeOperationTypes.h
@@ -103,7 +103,7 @@ namespace UE::MultiUserClient
 	struct FPropertyChange
 	{
 		/** The properties of the operation */
-		TArray<FConcertPropertyChain> Properties;
+		TSet<FConcertPropertyChain> Properties;
 		/** How to interpret Properties. */
 		EPropertyChangeType ChangeType;
 	};
diff --git a/Engine/Plugins/Developer/Concert/ConcertSync/ConcertSyncCore/Source/ConcertSyncCore/Private/Replication/AuthorityConflictSharedUtils.cpp b/Engine/Plugins/Developer/Concert/ConcertSync/ConcertSyncCore/Source/ConcertSyncCore/Private/Replication/AuthorityConflictSharedUtils.cpp
--- a/Engine/Plugins/Developer/Concert/ConcertSync/ConcertSyncCore/Source/ConcertSyncCore/Private/Replication/AuthorityConflictSharedUtils.cpp
+++ b/Engine/Plugins/Developer/Concert/ConcertSync/ConcertSyncCore/Source/ConcertSyncCore/Private/Replication/AuthorityConflictSharedUtils.cpp
@@ -53,7 +53,7 @@ namespace UE::ConcertSyncCore::Replication::AuthorityConflictUtils
 	EAuthorityConflict EnumerateAuthorityConflicts(
 		const FGuid& ClientId,
 		const FSoftObjectPath& Object,
-		TConstArrayView<FConcertPropertyChain> OverwriteProperties,
+		const TSet<FConcertPropertyChain>& OverwriteProperties,
 		const IReplicationGroundTruth& GroundTruth,
 		FProcessAuthorityConflict ProcessConflict
 		)
@@ -101,7 +101,7 @@ namespace UE::ConcertSyncCore::Replication::AuthorityConflictUtils
 				}
 
 				const FSoftObjectPath& ObjectPath = Change.Key;
-				const TArray<FConcertPropertyChain>& Properties = ObjectInfo->PropertySelection.ReplicatedProperties;
+				const TSet<FConcertPropertyChain>& Properties = ObjectInfo->PropertySelection.ReplicatedProperties;
 				const EAuthorityConflict Conflict = EnumerateAuthorityConflicts(SendingClient, ObjectPath, Properties, GroundTruth);
 
 				const bool bHasConflict = Conflict == EAuthorityConflict::Conflict;
diff --git a/Engine/Plugins/Developer/Concert/ConcertSync/ConcertSyncCore/Source/ConcertSyncCore/Private/Replication/Data/ConcertPropertySelection.cpp b/Engine/Plugins/Developer/Concert/ConcertSync/ConcertSyncCore/Source/ConcertSyncCore/Private/Replication/Data/ConcertPropertySelection.cpp
--- a/Engine/Plugins/Developer/Concert/ConcertSync/ConcertSyncCore/Source/ConcertSyncCore/Private/Replication/Data/ConcertPropertySelection.cpp
+++ b/Engine/Plugins/Developer/Concert/ConcertSync/ConcertSyncCore/Source/ConcertSyncCore/Private/Replication/Data/ConcertPropertySelection.cpp
@@ -10,251 +10,226 @@
 
 TOptional<FConcertPropertyChain> FConcertPropertyChain::CreateFromPath(const UStruct& Class, const TArray<FName>& NamePath)
 {
-	TOptional<FConcertPropertyChain> Result;
-	// The goal here is to find a valid FProperty path based on the NamePath...
-	UE::ConcertSyncCore::PropertyChain::ForEachReplicatableConcertProperty(Class, [&NamePath, &Result](FConcertPropertyChain&& Chain) mutable
-	{
-		// ... and performance could be improved by not expanding down the search tree by introducing ETreeBreakBehavior::SkipSubtree
-		if (Chain == NamePath)
-		{
-			Result.Emplace(MoveTemp(Chain));
-			return EBreakBehavior::Break;
-		}
-		return EBreakBehavior::Continue;
-	});
-	return Result;
+    TOptional<FConcertPropertyChain> Result;
+    // The goal here is to find a valid FProperty path based on the NamePath...
+    UE::ConcertSyncCore::PropertyChain::ForEachReplicatableConcertProperty(Class, [&NamePath, &Result](FConcertPropertyChain&& Chain) mutable
+    {
+       // ... and performance could be improved by not expanding down the search tree by introducing ETreeBreakBehavior::SkipSubtree
+       if (Chain == NamePath)
+       {
+          Result.Emplace(MoveTemp(Chain));
+          return EBreakBehavior::Break;
+       }
+       return EBreakBehavior::Continue;
+    });
+    return Result;
 }
 
 FConcertPropertyChain::FConcertPropertyChain(const FArchiveSerializedPropertyChain* OptionalChain, const FProperty& LeafProperty)
 {
-	using namespace UE::ConcertSyncCore::PropertyChain;
-	
-	if (!ensureMsgf(!IsInnerContainerProperty(LeafProperty), TEXT("Concert property chains never contain inner properties so you cannot construct one to point at it!")))
-	{
-		return;
-	}
-	
-	if (OptionalChain)
-	{
-		PathToProperty.Reserve(OptionalChain->GetNumProperties() + 1);
-		for (int32 i = 0; i < OptionalChain->GetNumProperties(); ++i)
-		{
-			const FProperty* CurrentProperty = OptionalChain->GetPropertyFromRoot(i);
-			
-			const FProperty* ParentProperty = CurrentProperty->GetOwner<FProperty>();
-			const bool bIsKeyOfMap = ParentProperty && ParentProperty->IsA(FMapProperty::StaticClass()) && CastField<FMapProperty>(ParentProperty)->KeyProp == CurrentProperty;
-			if (!ensureMsgf(!bIsKeyOfMap, TEXT("Key child properties never appear in a path because all of them are assumed to be replicated. See documentation of PathToProperty!")))
-			{
-				// The input was invalid but we've constructed FConcertPropertyChain with valid state: it ends at the TMap property.
-				return;
-			}
-
-			// Inner properties do not show up in the path but are expected to be contained in OptionalChain since it's needed for the search.
-			if (!IsInnerContainerProperty(*CurrentProperty))
-			{
-				PathToProperty.Add(OptionalChain->GetPropertyFromRoot(i)->GetFName());
-			}
-		}
-	}
-	else if (!ensureMsgf(LeafProperty.GetOwner<FProperty>() == nullptr, TEXT("Not a leaf property!")))
-	{
-		// The input was invalid but we've constructed FConcertPropertyChain with valid state: no property.
-		return;
-	}
-	
-	PathToProperty.Add(LeafProperty.GetFName());
-	UE_CLOG(!IsReplicatableProperty(LeafProperty), LogConcert, Warning, TEXT("Instantiated property chain with non replicatable property: %s"), *ToString());
+    using namespace UE::ConcertSyncCore::PropertyChain;
+    
+    if (!ensureMsgf(!IsInnerContainerProperty(LeafProperty), TEXT("Concert property chains never contain inner properties so you cannot construct one to point at it!")))
+    {
+       return;
+    }
+    
+    if (OptionalChain)
+    {
+       PathToProperty.Reserve(OptionalChain->GetNumProperties() + 1);
+       for (int32 i = 0; i < OptionalChain->GetNumProperties(); ++i)
+       {
+          const FProperty* CurrentProperty = OptionalChain->GetPropertyFromRoot(i);
+          
+          const FProperty* ParentProperty = CurrentProperty->GetOwner<FProperty>();
+          const bool bIsKeyOfMap = ParentProperty && ParentProperty->IsA(FMapProperty::StaticClass()) && CastField<FMapProperty>(ParentProperty)->KeyProp == CurrentProperty;
+          if (!ensureMsgf(!bIsKeyOfMap, TEXT("Key child properties never appear in a path because all of them are assumed to be replicated. See documentation of PathToProperty!")))
+          {
+             // The input was invalid but we've constructed FConcertPropertyChain with valid state: it ends at the TMap property.
+             return;
+          }
+
+          // Inner properties do not show up in the path but are expected to be contained in OptionalChain since it's needed for the search.
+          if (!IsInnerContainerProperty(*CurrentProperty))
+          {
+             PathToProperty.Add(OptionalChain->GetPropertyFromRoot(i)->GetFName());
+          }
+       }
+    }
+    else if (!ensureMsgf(LeafProperty.GetOwner<FProperty>() == nullptr, TEXT("Not a leaf property!")))
+    {
+       // The input was invalid but we've constructed FConcertPropertyChain with valid state: no property.
+       return;
+    }
+    
+    PathToProperty.Add(LeafProperty.GetFName());
+    UE_CLOG(!IsReplicatableProperty(LeafProperty), LogConcert, Warning, TEXT("Instantiated property chain with non replicatable property: %s"), *ToString());
 }
 
 bool FConcertPropertyChain::IsChildOf(const FConcertPropertyChain& ParentToCheck) const
 {
-	if (PathToProperty.Num() < ParentToCheck.PathToProperty.Num())
-	{
-		return false;
-	}
-
-	for (int32 i = 0; i < ParentToCheck.PathToProperty.Num(); ++i)
-	{
-		if (PathToProperty[i] != ParentToCheck.PathToProperty[i])
-		{
-			return false;
-		}
-	}
-
-	return true;
+    if (PathToProperty.Num() < ParentToCheck.PathToProperty.Num())
+    {
+       return false;
+    }
+
+    for (int32 i = 0; i < ParentToCheck.PathToProperty.Num(); ++i)
+    {
+       if (PathToProperty[i] != ParentToCheck.PathToProperty[i])
+       {
+          return false;
+       }
+    }
+
+    return true;
 }
 
 bool FConcertPropertyChain::IsDirectChildOf(const FConcertPropertyChain& ParentToCheck) const
 {
-	if (ParentToCheck.PathToProperty.Num() != PathToProperty.Num() - 1)
-	{
-		return false;
-	}
-
-	bool bArePathsEqual = true;
-	for (int32 i = 0; bArePathsEqual && i < ParentToCheck.PathToProperty.Num(); ++i)
-	{
-		bArePathsEqual &= ParentToCheck.PathToProperty[i] == PathToProperty[i];
-	}
-	return bArePathsEqual;
+    if (ParentToCheck.PathToProperty.Num() != PathToProperty.Num() - 1)
+    {
+       return false;
+    }
+
+    bool bArePathsEqual = true;
+    for (int32 i = 0; bArePathsEqual && i < ParentToCheck.PathToProperty.Num(); ++i)
+    {
+       bArePathsEqual &= ParentToCheck.PathToProperty[i] == PathToProperty[i];
+    }
+    return bArePathsEqual;
 }
 
 bool FConcertPropertyChain::MatchesExactly(const FArchiveSerializedPropertyChain* OptionalChain, const FProperty& LeafProperty) const
 {
-	using namespace UE::ConcertSyncCore::PropertyChain;
-	if (PathToProperty.IsEmpty())
-	{
-		return false;
-	}
-	
-	const int32 OptionalChainLength = OptionalChain ? OptionalChain->GetNumProperties() : 0;
-	const bool bLeafPropertiesMatch = LeafProperty.GetFName() == PathToProperty[PathToProperty.Num() - 1];
-	// No point continuing matching properties if leaf properties do not match.
-	if (!bLeafPropertiesMatch)
-	{
-		return false;
-	}
-
-	// Does the chain describe a single property? If yes, we're done.
-	if (OptionalChainLength == 0 && PathToProperty.Num() == 1)
-	{
-		return true;
-	}
-
-	// We get here if PathToproperty.Num() > 1 so OptionalChain must be valid, otherwise we're not equal.
-	if (!OptionalChain)
-	{
-		return false;
-	}
-
-	/*
-	 * OptionalChainLength can be larger than PathToProperty.Num() - 1 and still match!
-	 * 
-	 * Example:
-	 *	USTRUCT() struct FInner { UPROPERTY() float Value; }
-	 *	UCLASS() class UFoo : public UObject { UPROPERTY() TArray<FInner> Nested; }
-	 *	
-	 * Supposing PathToProperty = { "Nested", "Value" }, this is what the input to Matches would be:
-	 *  - OptionalChain = { "Nested", "Nested" } ==> 1st Nested is the FArrayProperty, 2nd the FArrayProperty::Inner
-	 *  - LeafProperty = "Value"
-	 */	
-	bool bArePathsEqual = true;
-
-	int32 Index_PathToProperty = 0;
-	int32 Index_OptionalChain = 0;
-	// Advance Index_OptionalChain towards the end of OptionalChain...
-	for (; bArePathsEqual && PathToProperty.IsValidIndex(Index_PathToProperty) && Index_OptionalChain < OptionalChain->GetNumProperties(); ++Index_OptionalChain)
-	{
-		const FProperty* OptionalChainProperty = OptionalChain->GetPropertyFromRoot(Index_OptionalChain);
-		const bool bIsInternalContainerProperty = IsInnerContainerProperty(*OptionalChainProperty);
-		bArePathsEqual &= bIsInternalContainerProperty || OptionalChainProperty->GetFName() == PathToProperty[Index_PathToProperty];
-
-		// ... but only advance Index_PathToProperty if the current property is not an inner container property since they do not show up in FConcertPropertyChains.
-		if (!bIsInternalContainerProperty)
-		{
-			++Index_PathToProperty;
-		}
-	}
-
-	// If not every property was visited, that could mean one path is is a sub-path of the other. Regardless of the reason, they stopped matching somewhere in the middle so return false.
-	const bool bVisitedEveryProperty = Index_OptionalChain == OptionalChain->GetNumProperties()
-		// -1 to exclude leaf property (we checked it way at the beginning): OptionalChain does not contain the leaf property.
-		&& (Index_PathToProperty == PathToProperty.Num() - 1);
-	return bArePathsEqual && bVisitedEveryProperty;
+    using namespace UE::ConcertSyncCore::PropertyChain;
+    if (PathToProperty.IsEmpty())
+    {
+       return false;
+    }
+    
+    const int32 OptionalChainLength = OptionalChain ? OptionalChain->GetNumProperties() : 0;
+    const bool bLeafPropertiesMatch = LeafProperty.GetFName() == PathToProperty[PathToProperty.Num() - 1];
+    // No point continuing matching properties if leaf properties do not match.
+    if (!bLeafPropertiesMatch)
+    {
+       return false;
+    }
+
+    // Does the chain describe a single property? If yes, we're done.
+    if (OptionalChainLength == 0 && PathToProperty.Num() == 1)
+    {
+       return true;
+    }
+
+    // We get here if PathToproperty.Num() > 1 so OptionalChain must be valid, otherwise we're not equal.
+    if (!OptionalChain)
+    {
+       return false;
+    }
+
+    /*
+     * OptionalChainLength can be larger than PathToProperty.Num() - 1 and still match!
+     * 
+     * Example:
+     * USTRUCT() struct FInner { UPROPERTY() float Value; }
+     * UCLASS() class UFoo : public UObject { UPROPERTY() TArray<FInner> Nested; }
+     * 
+     * Supposing PathToProperty = { "Nested", "Value" }, this is what the input to Matches would be:
+     *  - OptionalChain = { "Nested", "Nested" } ==> 1st Nested is the FArrayProperty, 2nd the FArrayProperty::Inner
+     *  - LeafProperty = "Value"
+     */    
+    bool bArePathsEqual = true;
+
+    int32 Index_PathToProperty = 0;
+    int32 Index_OptionalChain = 0;
+    // Advance Index_OptionalChain towards the end of OptionalChain...
+    for (; bArePathsEqual && PathToProperty.IsValidIndex(Index_PathToProperty) && Index_OptionalChain < OptionalChain->GetNumProperties(); ++Index_OptionalChain)
+    {
+       const FProperty* OptionalChainProperty = OptionalChain->GetPropertyFromRoot(Index_OptionalChain);
+       const bool bIsInternalContainerProperty = IsInnerContainerProperty(*OptionalChainProperty);
+       bArePathsEqual &= bIsInternalContainerProperty || OptionalChainProperty->GetFName() == PathToProperty[Index_PathToProperty];
+
+       // ... but only advance Index_PathToProperty if the current property is not an inner container property since they do not show up in FConcertPropertyChains.
+       if (!bIsInternalContainerProperty)
+       {
+          ++Index_PathToProperty;
+       }
+    }
+
+    // If not every property was visited, that could mean one path is is a sub-path of the other. Regardless of the reason, they stopped matching somewhere in the middle so return false.
+    const bool bVisitedEveryProperty = Index_OptionalChain == OptionalChain->GetNumProperties()
+       // -1 to exclude leaf property (we checked it way at the beginning): OptionalChain does not contain the leaf property.
+       && (Index_PathToProperty == PathToProperty.Num() - 1);
+    return bArePathsEqual && bVisitedEveryProperty;
 }
 
 FProperty* FConcertPropertyChain::ResolveProperty(UStruct& Class, bool bLogOnFail) const
 {
-	// FConcertPropertyChain::ResolveProperty exists only for visibility to developers since
-	// the class is the first place one would look and not in the utils namespace.
-	return UE::ConcertSyncCore::PropertyChain::ResolveProperty(Class, *this, bLogOnFail);
+    // FConcertPropertyChain::ResolveProperty exists only for visibility to developers since
+    // the class is the first place one would look and not in the utils namespace.
+    return UE::ConcertSyncCore::PropertyChain::ResolveProperty(Class, *this, bLogOnFail);
 }
 
 FConcertPropertyChain FConcertPropertyChain::GetParent() const
 {
-	if (PathToProperty.Num() <= 1)
-	{
-		return {};
-	}
-	
-	FConcertPropertyChain Result = *this;
-	Result.PathToProperty.RemoveAt(Result.PathToProperty.Num() - 1);
-	return Result;
+    if (PathToProperty.Num() <= 1)
+    {
+       return {};
+    }
+    
+    FConcertPropertyChain Result = *this;
+    Result.PathToProperty.RemoveAt(Result.PathToProperty.Num() - 1);
+    return Result;
 }
 
 FString FConcertPropertyChain::ToString(EToStringMethod Method) const
 {
-	switch (Method)
-	{
-	case EToStringMethod::Path:
-		return FString::JoinBy(PathToProperty, TEXT("."), [](const FName& Name){ return Name.ToString(); });
-	case EToStringMethod::LeafProperty:
-		return PathToProperty.IsEmpty()
-			? FName(NAME_None).ToString()
-			: PathToProperty[PathToProperty.Num() - 1].ToString();
-	default:
-		checkNoEntry();
-		return FName(NAME_None).ToString();
-	}
-}
-
-bool FConcertPropertySelection::Includes(const FConcertPropertySelection& Other) const
-{
-	if (ReplicatedProperties.Num() < Other.ReplicatedProperties.Num())
-	{
-		return false;
-	}
-
-	return Algo::AllOf(ReplicatedProperties, [&Other](const FConcertPropertyChain& Property)
-	{
-		return Other.ReplicatedProperties.Contains(Property);
-	});
+    switch (Method)
+    {
+    case EToStringMethod::Path:
+       return FString::JoinBy(PathToProperty, TEXT("."), [](const FName& Name){ return Name.ToString(); });
+    case EToStringMethod::LeafProperty:
+       return PathToProperty.IsEmpty()
+          ? FName(NAME_None).ToString()
+          : PathToProperty[PathToProperty.Num() - 1].ToString();
+    default:
+       checkNoEntry();
+       return FName(NAME_None).ToString();
+    }
 }
 
 void FConcertPropertySelection::DiscoverAndAddImplicitParentProperties()
 {
-	// Iterate in reverse to ignore the properties we're about to add
-	for (int32 Index = ReplicatedProperties.Num() - 1; ReplicatedProperties.IsValidIndex(Index); ++Index)
-	{
-		FConcertPropertyChain Current = ReplicatedProperties[Index].GetParent();
-		while (!Current.IsEmpty())
-		{
-			// This part makes the algorithm O(n^2) but could be O(n) with a TSet. Suboptimal, but it's simple... it's quite unlikely this will cause us performance trouble
-			ReplicatedProperties.AddUnique(Current);
-			Current = Current.GetParent();
-		}
-	}
+    // Iterate in reverse to ignore the properties we're about to add
+    for (FConcertPropertyChain Current : TSet(ReplicatedProperties))
+    {
+       while (!Current.IsEmpty())
+       {
+          ReplicatedProperties.Add(Current);
+          Current = Current.GetParent();
+       }
+    }
 }
 
 bool FConcertPropertySelection::EnumeratePropertyOverlaps(
-	TConstArrayView<FConcertPropertyChain> First,
-	TConstArrayView<FConcertPropertyChain> Second,
-	TFunctionRef<EBreakBehavior(const FConcertPropertyChain&)> Callback
-	)
+    const TSet<FConcertPropertyChain>& First,
+    const TSet<FConcertPropertyChain>& Second,
+    TFunctionRef<EBreakBehavior(const FConcertPropertyChain&)> Callback
+    )
 {
-	/* 
-	 * Performance could be improved: Two chains can only overlap if their root properties overlap.
-	 * Example: [Foo.Vector.X] and [Foo.FloatProperty] both share Foo struct (but do not overlap in this case).
-	 */
-	bool bHasOverlap = false;
-	for (const FConcertPropertyChain& ThisChain : First)
-	{
-		for (const FConcertPropertyChain& OtherChain : Second)
-		{
-			const bool bIsOverlap = ThisChain == OtherChain;
-			bHasOverlap |= bIsOverlap;
-			if (bIsOverlap && Callback(ThisChain) == EBreakBehavior::Break) 
-			{
-				return true;
-			}
-		}
-	}
-
-	return bHasOverlap;
+    const TSet<FConcertPropertyChain> Intersection = First.Intersect(Second);
+    for (const FConcertPropertyChain& ThisChain : Intersection)
+    {
+       if (Callback(ThisChain) == EBreakBehavior::Break) 
+       {
+          return true;
+       }
+    }
+
+    return !Intersection.IsEmpty();
 }
 
 uint32 GetTypeHash(const FConcertPropertyChain& Chain)
 {
-	return GetTypeHash(Chain.GetPathToProperty());
-}
-
+    return GetTypeHash(Chain.GetPathToProperty());
+}
\ No newline at end of file
diff --git a/Engine/Plugins/Developer/Concert/ConcertSync/ConcertSyncCore/Source/ConcertSyncCore/Private/Replication/Messages/Handshake.cpp b/Engine/Plugins/Developer/Concert/ConcertSync/ConcertSyncCore/Source/ConcertSyncCore/Private/Replication/Messages/Handshake.cpp
--- a/Engine/Plugins/Developer/Concert/ConcertSync/ConcertSyncCore/Source/ConcertSyncCore/Private/Replication/Messages/Handshake.cpp
+++ b/Engine/Plugins/Developer/Concert/ConcertSync/ConcertSyncCore/Source/ConcertSyncCore/Private/Replication/Messages/Handshake.cpp
@@ -6,7 +6,7 @@ namespace UE::ConcertSyncCore::Replication
 {
 	FString LexJoinErrorCode(EJoinReplicationErrorCode ErrorCode)
 	{
-		static_assert(static_cast<int32>(EJoinReplicationErrorCode::Max) == 9, "Update JoinErrorCodeToText when you change entries");
+		static_assert(static_cast<int32>(EJoinReplicationErrorCode::Max) == 8, "Update JoinErrorCodeToText when you change entries");
 
 		switch (ErrorCode)
 		{
@@ -16,7 +16,6 @@ namespace UE::ConcertSyncCore::Replication
 		case EJoinReplicationErrorCode::AlreadyInProgress: return TEXT("AlreadyInProgress");
 		case EJoinReplicationErrorCode::NotInAnyConcertSession: return TEXT("NotInAnyConcertSession");
 		case EJoinReplicationErrorCode::InvalidClass: return TEXT("InvalidClass");
-		case EJoinReplicationErrorCode::DuplicateProperty: return TEXT("DuplicateProperty");
 		case EJoinReplicationErrorCode::AlreadyInSession: return TEXT("AlreadyInSession");
 		case EJoinReplicationErrorCode::DuplicateStreamId: return TEXT("DuplicateStreamId");
 		case EJoinReplicationErrorCode::FailedToUnpackStream: return TEXT("FailedToUnpackStream");
diff --git a/Engine/Plugins/Developer/Concert/ConcertSync/ConcertSyncCore/Source/ConcertSyncCore/Private/Replication/ReplicationPropertyFilter.cpp b/Engine/Plugins/Developer/Concert/ConcertSync/ConcertSyncCore/Source/ConcertSyncCore/Private/Replication/ReplicationPropertyFilter.cpp
--- a/Engine/Plugins/Developer/Concert/ConcertSync/ConcertSyncCore/Source/ConcertSyncCore/Private/Replication/ReplicationPropertyFilter.cpp
+++ b/Engine/Plugins/Developer/Concert/ConcertSync/ConcertSyncCore/Source/ConcertSyncCore/Private/Replication/ReplicationPropertyFilter.cpp
@@ -11,10 +11,9 @@ namespace UE::ConcertSyncCore
 	FReplicationPropertyFilter::FReplicationPropertyFilter(const FConcertPropertySelection& PropertySelection)
 		: PropertySelection(PropertySelection)
 	{
-		for (int32 i = 0; i < PropertySelection.ReplicatedProperties.Num(); ++i)
+		for (const FConcertPropertyChain& Chain : PropertySelection.ReplicatedProperties)
 		{
-			const FConcertPropertyChain& Chain = PropertySelection.ReplicatedProperties[i];
-			LeafToChain.FindOrAdd(Chain.GetLeafProperty()).Add(i);
+			LeafToChain.FindOrAdd(Chain.GetLeafProperty()).Add(&Chain);
 		}
 	}
 
@@ -38,8 +37,8 @@ namespace UE::ConcertSyncCore
 			return true;
 		}
 		
-		const TArray<int32>* IndicesToSearch = LeafToChain.Find(Property.GetFName());
-		if (!IndicesToSearch)
+		const TArray<const FConcertPropertyChain*>* ChainsEndingWithPropertyName = LeafToChain.Find(Property.GetFName());
+		if (!ChainsEndingWithPropertyName)
 		{
 			/* No chain ends with this property so it is not in the selection
 			 * Note: The property selection must also contain every parent property.
@@ -49,10 +48,9 @@ namespace UE::ConcertSyncCore
 			return false;
 		}
 
-		for (const int32 IndexToSearch : *IndicesToSearch)
+		for (const FConcertPropertyChain* ConcertChain : *ChainsEndingWithPropertyName)
 		{
-			const FConcertPropertyChain& ConcertChain = PropertySelection.ReplicatedProperties[IndexToSearch];
-			if (ConcertChain.MatchesExactly(Chain, Property))
+			if (ConcertChain->MatchesExactly(Chain, Property))
 			{
 				return true;
 			}
diff --git a/Engine/Plugins/Developer/Concert/ConcertSync/ConcertSyncCore/Source/ConcertSyncCore/Public/Replication/AuthorityConflictSharedUtils.h b/Engine/Plugins/Developer/Concert/ConcertSync/ConcertSyncCore/Source/ConcertSyncCore/Public/Replication/AuthorityConflictSharedUtils.h
--- a/Engine/Plugins/Developer/Concert/ConcertSync/ConcertSyncCore/Source/ConcertSyncCore/Public/Replication/AuthorityConflictSharedUtils.h
+++ b/Engine/Plugins/Developer/Concert/ConcertSync/ConcertSyncCore/Source/ConcertSyncCore/Public/Replication/AuthorityConflictSharedUtils.h
@@ -59,7 +59,7 @@ namespace UE::ConcertSyncCore::Replication::AuthorityConflictUtils
 	CONCERTSYNCCORE_API EAuthorityConflict EnumerateAuthorityConflicts(
 		const FGuid& ClientId,
 		const FSoftObjectPath& Object,
-		TConstArrayView<FConcertPropertyChain> OverwriteProperties,
+		const TSet<FConcertPropertyChain>& OverwriteProperties,
 		const IReplicationGroundTruth& GroundTruth,
 		FProcessAuthorityConflict ProcessConflict = [](const FGuid&, const FGuid&, const FConcertPropertyChain&){ return EBreakBehavior::Break; }
 		);
diff --git a/Engine/Plugins/Developer/Concert/ConcertSync/ConcertSyncCore/Source/ConcertSyncCore/Public/Replication/Data/ConcertPropertySelection.h b/Engine/Plugins/Developer/Concert/ConcertSync/ConcertSyncCore/Source/ConcertSyncCore/Public/Replication/Data/ConcertPropertySelection.h
--- a/Engine/Plugins/Developer/Concert/ConcertSync/ConcertSyncCore/Source/ConcertSyncCore/Public/Replication/Data/ConcertPropertySelection.h
+++ b/Engine/Plugins/Developer/Concert/ConcertSync/ConcertSyncCore/Source/ConcertSyncCore/Public/Replication/Data/ConcertPropertySelection.h
@@ -148,13 +148,13 @@ struct CONCERTSYNCCORE_API FConcertPropertySelection
 
 	/** List of replicated properties. */
 	UPROPERTY()
-	TArray<FConcertPropertyChain> ReplicatedProperties;
+	TSet<FConcertPropertyChain> ReplicatedProperties;
 
 	/** @return Whether this and Other contain at least one property that is the same. */
 	bool OverlapsWith(const FConcertPropertySelection& Other) const { return EnumeratePropertyOverlaps(ReplicatedProperties, Other.ReplicatedProperties); }
 
 	/** @return Whether this includes all properties of Other */
-	bool Includes(const FConcertPropertySelection& Other) const;
+	bool Includes(const FConcertPropertySelection& Other) const { return ReplicatedProperties.Includes(Other.ReplicatedProperties); }
 
 	/**
 	 * Adds all parent properties if they are missing.
@@ -171,14 +171,14 @@ struct CONCERTSYNCCORE_API FConcertPropertySelection
 	 * @return Whether there were any property overlaps.
 	 */
 	static bool EnumeratePropertyOverlaps(
-		TConstArrayView<FConcertPropertyChain> First,
-		TConstArrayView<FConcertPropertyChain> Second,
+		const TSet<FConcertPropertyChain>& First,
+		const TSet<FConcertPropertyChain>& Second,
 		TFunctionRef<EBreakBehavior(const FConcertPropertyChain&)> Callback = [](const FConcertPropertyChain&){ return EBreakBehavior::Break; }
 		);
 	
 	friend bool operator==(const FConcertPropertySelection& Left, const FConcertPropertySelection& Right)
 	{
-		return Left.ReplicatedProperties == Right.ReplicatedProperties;
+		return Left.ReplicatedProperties.Num() == Right.ReplicatedProperties.Num() && Left.Includes(Right);
 	}
 	friend bool operator!=(const FConcertPropertySelection& Left, const FConcertPropertySelection& Right)
 	{
diff --git a/Engine/Plugins/Developer/Concert/ConcertSync/ConcertSyncCore/Source/ConcertSyncCore/Public/Replication/Messages/Handshake.h b/Engine/Plugins/Developer/Concert/ConcertSync/ConcertSyncCore/Source/ConcertSyncCore/Public/Replication/Messages/Handshake.h
--- a/Engine/Plugins/Developer/Concert/ConcertSync/ConcertSyncCore/Source/ConcertSyncCore/Public/Replication/Messages/Handshake.h
+++ b/Engine/Plugins/Developer/Concert/ConcertSync/ConcertSyncCore/Source/ConcertSyncCore/Public/Replication/Messages/Handshake.h
@@ -35,12 +35,10 @@ enum class EJoinReplicationErrorCode : uint8
 	// Invalid input:
 	/** One of the specified classes was null or otherwise invalid */
 	InvalidClass = 6,
-	/** An object selection contained a property twice. */
-	DuplicateProperty = 7,
 	/** Your input contained the same stream id twice. */
-	DuplicateStreamId = 8,
+	DuplicateStreamId = 7,
 	/** Failed to unpack the stream, likely because one of the attribute classes could not be resolved. */
-	FailedToUnpackStream = 9,
+	FailedToUnpackStream = 8,
 	
 	MaxPlusOne,
 	/** The maximum possible entry value */
diff --git a/Engine/Plugins/Developer/Concert/ConcertSync/ConcertSyncCore/Source/ConcertSyncCore/Public/Replication/ReplicationPropertyFilter.h b/Engine/Plugins/Developer/Concert/ConcertSync/ConcertSyncCore/Source/ConcertSyncCore/Public/Replication/ReplicationPropertyFilter.h
--- a/Engine/Plugins/Developer/Concert/ConcertSync/ConcertSyncCore/Source/ConcertSyncCore/Public/Replication/ReplicationPropertyFilter.h
+++ b/Engine/Plugins/Developer/Concert/ConcertSync/ConcertSyncCore/Source/ConcertSyncCore/Public/Replication/ReplicationPropertyFilter.h
@@ -9,6 +9,7 @@
 
 class FProperty;
 struct FArchiveSerializedPropertyChain;
+struct FConcertPropertyChain;
 struct FConcertPropertySelection;
 
 namespace UE::ConcertSyncCore
@@ -31,10 +32,10 @@ namespace UE::ConcertSyncCore
 		const FConcertPropertySelection& PropertySelection;
 		
 		/**
-		 * Maps the name of every leaf property to the chain indices of PropertySelection.ReplicatedProperties that contain the property name at the end of the chain.
+		 * Maps the name of every leaf property to the chaisn PropertySelection.ReplicatedProperties that contain the property name at the end of the chain.
 		 * This speeds up matching chain and property.
 		 */
-		TMap<FName, TArray<int32>> LeafToChain;
+		TMap<FName, TArray<const FConcertPropertyChain*>> LeafToChain;
 	public:
 
 		FReplicationPropertyFilter(const FConcertPropertySelection& PropertySelection);
diff --git a/Engine/Plugins/Developer/Concert/ConcertSync/ConcertSyncServer/Source/ConcertSyncServer/Private/Replication/Util/JoinRequestValidation.cpp b/Engine/Plugins/Developer/Concert/ConcertSync/ConcertSyncServer/Source/ConcertSyncServer/Private/Replication/Util/JoinRequestValidation.cpp
--- a/Engine/Plugins/Developer/Concert/ConcertSync/ConcertSyncServer/Source/ConcertSyncServer/Private/Replication/Util/JoinRequestValidation.cpp
+++ b/Engine/Plugins/Developer/Concert/ConcertSync/ConcertSyncServer/Source/ConcertSyncServer/Private/Replication/Util/JoinRequestValidation.cpp
@@ -45,24 +45,6 @@ namespace UE::ConcertSyncServer::Replication
 						);
 					return { EJoinReplicationErrorCode::InvalidClass, ErrorMessage };
 				}
-
-				// Validate unique properties
-				TSet<FConcertPropertyChain> UniquePropertyDetection;
-				for (const FConcertPropertyChain& PropertyChain : ObjectInfo.PropertySelection.ReplicatedProperties)
-				{
-					// No duplicate properties!
-					if (UniquePropertyDetection.Contains(PropertyChain))
-					{
-						const FString ErrorMessage = FString::Printf(
-							TEXT("Stream %s has duplicate property %s for object %s"),
-							*Stream.BaseDescription.Identifier.ToString(),
-							*PropertyChain.ToString(),
-							*ObjectPath.ToString()
-						);
-						return { EJoinReplicationErrorCode::DuplicateProperty, ErrorMessage };
-					}
-					UniquePropertyDetection.Add(PropertyChain);
-				}
 			}
 		}
 
diff --git a/Engine/Plugins/Developer/Concert/ConcertSync/ConcertSyncTest/Source/ConcertSyncTest/Private/Replication/ConcertPropertyTests.cpp b/Engine/Plugins/Developer/Concert/ConcertSync/ConcertSyncTest/Source/ConcertSyncTest/Private/Replication/ConcertPropertyTests.cpp
--- a/Engine/Plugins/Developer/Concert/ConcertSync/ConcertSyncTest/Source/ConcertSyncTest/Private/Replication/ConcertPropertyTests.cpp
+++ b/Engine/Plugins/Developer/Concert/ConcertSync/ConcertSyncTest/Source/ConcertSyncTest/Private/Replication/ConcertPropertyTests.cpp
@@ -326,7 +326,7 @@ namespace UE::ConcertSyncTest::Replication::PropertyChain
 		};
 
 		// CreateFromPath is supposed to work (assuming the other tests worked so far) but we'll use TOptional::Get() calls are in case it does not. 
-		TArray<FConcertPropertyChain> AllowedProperties = {
+		TSet<FConcertPropertyChain> AllowedProperties = {
 			FConcertPropertyChain::CreateFromPath(*UTestReflectionObject::StaticClass(), { TEXT("Float") }).Get({}),
 			FConcertPropertyChain::CreateFromPath(*UTestReflectionObject::StaticClass(), { TEXT("Vector") }).Get({}),
 			FConcertPropertyChain::CreateFromPath(*UTestReflectionObject::StaticClass(), { TEXT("Vector"), TEXT("X") }).Get({}),
diff --git a/Engine/Plugins/Developer/Concert/ConcertSync/ConcertSyncTest/Source/ConcertSyncTest/Private/Replication/ReplicationAuthorityTests.cpp b/Engine/Plugins/Developer/Concert/ConcertSync/ConcertSyncTest/Source/ConcertSyncTest/Private/Replication/ReplicationAuthorityTests.cpp
--- a/Engine/Plugins/Developer/Concert/ConcertSync/ConcertSyncTest/Source/ConcertSyncTest/Private/Replication/ReplicationAuthorityTests.cpp
+++ b/Engine/Plugins/Developer/Concert/ConcertSync/ConcertSyncTest/Source/ConcertSyncTest/Private/Replication/ReplicationAuthorityTests.cpp
@@ -420,7 +420,7 @@ namespace UE::ConcertSyncTests::Replication::Authority
 			}
 		};
 		
-		FConcertBaseStreamInfo MakeStream(const FGuid& StreamId, const FSoftObjectPath ObjectPath, TArray<FConcertPropertyChain> Properties)
+		FConcertBaseStreamInfo MakeStream(const FGuid& StreamId, const FSoftObjectPath ObjectPath, TSet<FConcertPropertyChain> Properties)
 		{
 			FConcertObjectReplicationMap ExistingReplicationMap;
 			const FConcertPropertySelection Selection{ Properties};
diff --git a/Engine/Plugins/Developer/Concert/ConcertSync/ConcertSyncTest/Source/ConcertSyncTest/Private/Replication/ReplicationHandshakeTests.cpp b/Engine/Plugins/Developer/Concert/ConcertSync/ConcertSyncTest/Source/ConcertSyncTest/Private/Replication/ReplicationHandshakeTests.cpp
--- a/Engine/Plugins/Developer/Concert/ConcertSync/ConcertSyncTest/Source/ConcertSyncTest/Private/Replication/ReplicationHandshakeTests.cpp
+++ b/Engine/Plugins/Developer/Concert/ConcertSync/ConcertSyncTest/Source/ConcertSyncTest/Private/Replication/ReplicationHandshakeTests.cpp
@@ -70,21 +70,13 @@ namespace UE::ConcertSyncTests::Replication::Handshake
 		// These should cases never happen if using the editor tools.
 		// However malicious users or those with custom C++ logic can send whatever they want.
 		
-		// 2.1.1 Duplicate properties
-		FConcertReplicationStream Invalid_DoublePropertyDescription = StreamDescription;
-		Invalid_DoublePropertyDescription.BaseDescription.ReplicationMap.ReplicatedObjects[PathToSomeActorComponent].PropertySelection.ReplicatedProperties.Add(ForcedLodModelProperty);
-		ClientReplicationManager_Primary->JoinReplicationSession({ { Invalid_DoublePropertyDescription } })
-			.Next([&](const FJoinReplicatedSessionResult& Result)
-			{
-				TestTrue(TEXT("Cannot contain same properties twice"), Result.ErrorCode == EJoinReplicationErrorCode::DuplicateProperty);
-			});
-		// 2.1.2 Duplicate stream identifier
+		// 2.1.1 Duplicate stream identifier
 		ClientReplicationManager_Primary->JoinReplicationSession({ { StreamDescription, StreamDescription } })
 			.Next([&](const FJoinReplicatedSessionResult& Result)
 			{
 				TestTrue(TEXT("Cannot contain stream ID twice"), Result.ErrorCode == EJoinReplicationErrorCode::DuplicateStreamId);
 			});
-		// 2.1.3 Missing class path
+		// 2.1.2 Missing class path
 		ClientReplicationManager_Primary->JoinReplicationSession({ { InvalidClassStreamDescription } })
 			.Next([&](const FJoinReplicatedSessionResult& Result)
 			{
diff --git a/Engine/Plugins/Developer/Concert/ConcertUI/ConcertSharedSlate/Source/ConcertSharedSlate/Private/Replication/Editor/Model/GenericReplicationStreamModel.cpp b/Engine/Plugins/Developer/Concert/ConcertUI/ConcertSharedSlate/Source/ConcertSharedSlate/Private/Replication/Editor/Model/GenericReplicationStreamModel.cpp
--- a/Engine/Plugins/Developer/Concert/ConcertUI/ConcertSharedSlate/Source/ConcertSharedSlate/Private/Replication/Editor/Model/GenericReplicationStreamModel.cpp
+++ b/Engine/Plugins/Developer/Concert/ConcertUI/ConcertSharedSlate/Source/ConcertSharedSlate/Private/Replication/Editor/Model/GenericReplicationStreamModel.cpp
@@ -16,14 +16,16 @@ namespace UE::ConcertSharedSlate
 	namespace Private
 	{
 		/** @return Whether any property was added */
-		static bool AddParentProperties(const UStruct& Class, const FConcertPropertyChain& PropertyToAdd, TArray<FConcertPropertyChain>& ReplicatedProperties)
+		static bool AddParentProperties(const UStruct& Class, const FConcertPropertyChain& PropertyToAdd, TSet<FConcertPropertyChain>& ReplicatedProperties)
 		{
 			bool bAddedAtLeastOne = false;
 			ConcertSyncCore::PropertyChain::ForEachReplicatableConcertProperty(Class, [&bAddedAtLeastOne, &ReplicatedProperties, &PropertyToAdd](FConcertPropertyChain&& Property)
 				{
 					if (Property.IsParentOf(PropertyToAdd))
 					{
-						bAddedAtLeastOne |= ReplicatedProperties.AddUnique(Property) != INDEX_NONE;
+						bool bWasAlreadyInSet = false;
+						ReplicatedProperties.Add(Property, &bWasAlreadyInSet);
+						bAddedAtLeastOne |= !bWasAlreadyInSet;
 					}
 					return EBreakBehavior::Continue;
 				});
@@ -236,11 +238,14 @@ namespace UE::ConcertSharedSlate
 		}
 
 		bool bAddedAtLeastOne = false;
-		TArray<FConcertPropertyChain>& ReplicatedProperties = AssignedProperties->PropertySelection.ReplicatedProperties;
+		TSet<FConcertPropertyChain>& ReplicatedProperties = AssignedProperties->PropertySelection.ReplicatedProperties;
 		ReplicatedProperties.Reserve(ReplicatedProperties.Num() + Properties.Num());
 		for (const FConcertPropertyChain& AddedProperty : Properties)
 		{
-			bAddedAtLeastOne |= ReplicatedProperties.AddUnique(AddedProperty) != INDEX_NONE;
+			bool bWasAlreadyInSet = false;
+			ReplicatedProperties.Add(AddedProperty, &bWasAlreadyInSet);
+			
+			bAddedAtLeastOne |= !bWasAlreadyInSet;
 			// Parent properties must also be added
 			// Not exactly efficient to iterate through the hierarchy for every removed item but it should be fine... Properties.Num() == 1 is the most common case
 			bAddedAtLeastOne |= Private::AddParentProperties(*Class, AddedProperty, ReplicatedProperties);
@@ -342,7 +347,7 @@ namespace UE::ConcertSharedSlate
 
 				if (!ObjectInfo.PropertySelection.ReplicatedProperties.Contains(PropertyChain))
 				{
-					TArray<FConcertPropertyChain>& Properties = ObjectInfo.PropertySelection.ReplicatedProperties;
+					TSet<FConcertPropertyChain>& Properties = ObjectInfo.PropertySelection.ReplicatedProperties;
 					Private::AddParentProperties(*ObjectClass, PropertyChain, Properties);
 					Properties.Emplace(MoveTemp(PropertyChain));
 				}
