Allow rewriter archives to use a different object path for the in-memory data vs the rewritten data

This can be needed to allow the in-memory state to serialize correctly while running it through the rewriter archive

#rnx

diff --git a/Engine/Plugins/Developer/Concert/ConcertSync/ConcertSyncCore/Source/ConcertSyncCore/Private/ConcertSyncArchives.cpp b/Engine/Plugins/Developer/Concert/ConcertSync/ConcertSyncCore/Source/ConcertSyncCore/Private/ConcertSyncArchives.cpp
--- a/Engine/Plugins/Developer/Concert/ConcertSync/ConcertSyncCore/Source/ConcertSyncCore/Private/ConcertSyncArchives.cpp
+++ b/Engine/Plugins/Developer/Concert/ConcertSync/ConcertSyncCore/Source/ConcertSyncCore/Private/ConcertSyncArchives.cpp
@@ -656,7 +656,9 @@ FArchive& FConcertSyncObjectRewriter::operator<<(UObject*& Obj)
 	{
 		OnObjectSerialized(ObjPath);
 		RewriteData(OffsetBeforeObjectRead, OffsetAfterObjectRead - OffsetBeforeObjectRead, ObjPath);
-		Obj = StaticFindObject(UObject::StaticClass(), nullptr, *ObjPath.ToString());
+
+		const FSoftObjectPath OutputObjPath = GetOutputObjectPath(ObjPath);
+		Obj = StaticFindObject(UObject::StaticClass(), nullptr, *OutputObjPath.ToString());
 	}
 
 	return *this;
@@ -706,7 +708,7 @@ FArchive& FConcertSyncObjectRewriter::operator<<(FSoftObjectPath& AssetPtr)
 	{
 		OnObjectSerialized(ObjPath);
 		RewriteData(OffsetBeforeObjectRead, OffsetAfterObjectRead - OffsetBeforeObjectRead, ObjPath);
-		AssetPtr = ObjPath;
+		AssetPtr = GetOutputObjectPath(ObjPath);
 	}
 
 	return *this;
diff --git a/Engine/Plugins/Developer/Concert/ConcertSync/ConcertSyncCore/Source/ConcertSyncCore/Public/ConcertSyncArchives.h b/Engine/Plugins/Developer/Concert/ConcertSync/ConcertSyncCore/Source/ConcertSyncCore/Public/ConcertSyncArchives.h
--- a/Engine/Plugins/Developer/Concert/ConcertSync/ConcertSyncCore/Source/ConcertSyncCore/Public/ConcertSyncArchives.h
+++ b/Engine/Plugins/Developer/Concert/ConcertSync/ConcertSyncCore/Source/ConcertSyncCore/Public/ConcertSyncArchives.h
@@ -4,6 +4,7 @@
 
 #include "Delegates/Delegate.h"
 #include "IdentifierTable/ConcertTransportArchives.h"
+#include "UObject/SoftObjectPath.h"
 
 struct FLazyObjectPtr;
 struct FObjectPtr;
@@ -147,7 +148,11 @@ public:
 	//~ End FArchive Interface
 
 protected:
+	/** Called to rewrite the object path that will be stored in the persistent object value (InBytes) */
 	virtual void OnObjectSerialized(FSoftObjectPath& Obj) {}
+
+	/** Called with the result of OnObjectSerialized to get the object path to set on the in-memory object value (the output of operator<<) */
+	virtual FSoftObjectPath GetOutputObjectPath(const FSoftObjectPath& Obj) { return Obj; }
 };
 
 #if UE_ENABLE_INCLUDE_ORDER_DEPRECATED_IN_5_2
