Update database code to allow adding of FConcertSyncReplicationActivity:
- Introduce FConcertSyncConnectionActivity (will be used in later CL)
- Extend EConcertSyncActivityEventType with Replication and Count; EConcertSyncActivityEventType::Count is used by static asserts (to make it easier for future devs to add new types)
- Update ConcertDatabaseTests.cpp and add dedicated ReplicationAcitivity spec test

diff --git a/Engine/Plugins/Developer/Concert/ConcertSync/ConcertSyncCore/Source/ConcertSyncCore/Private/ConcertSyncSessionDatabase.cpp b/Engine/Plugins/Developer/Concert/ConcertSync/ConcertSyncCore/Source/ConcertSyncCore/Private/ConcertSyncSessionDatabase.cpp
--- a/Engine/Plugins/Developer/Concert/ConcertSync/ConcertSyncCore/Source/ConcertSyncCore/Private/ConcertSyncSessionDatabase.cpp
+++ b/Engine/Plugins/Developer/Concert/ConcertSync/ConcertSyncCore/Source/ConcertSyncCore/Private/ConcertSyncSessionDatabase.cpp
@@ -10,6 +10,7 @@
 #include "SQLiteDatabase.h"
 #include "HAL/FileManager.h"
 #include "SQLitePreparedStatement.h"
+#include "Replication/Messages/ReplicationActivity.h"
 #include "Serialization/MemoryReader.h"
 #include "Serialization/MemoryWriter.h"
 #include "Tasks/Task.h"
@@ -548,7 +549,11 @@ enum class FConcertSyncSessionDatabaseVersion
 	V2 = 2,
 	/** Added FConcertSyncActivity::Flags to the activities table */
 	AddActivityFlags = 3,
-	Current = AddActivityFlags,
+	/** Added FConcertSyncReplicationActivity, which added the replication_events table. */
+	AddReplicationActivities = 4,
+
+	LatestPlusOne,
+	Current = LatestPlusOne - 1,
 };
 
 class FConcertSyncSessionDatabaseStatements
@@ -611,6 +616,11 @@ public:
 		PREPARE_STATEMENT(Statement_GetMaxPackageEventIdAndTransactionEventIdAtSavePerPackageNameId);
 		PREPARE_STATEMENT(Statement_GetPackageHeadRevison);
 		PREPARE_STATEMENT(Statement_GetPackageTransactionEventIdAtLastSave);
+		
+		PREPARE_STATEMENT(Statement_AddReplicationData);
+		PREPARE_STATEMENT(Statement_SetReplicationData);
+		PREPARE_STATEMENT(Statement_GetReplicationEventForId);
+		PREPARE_STATEMENT(Statement_GetReplicationMaxEventIdByClientAndType);
 
 		PREPARE_STATEMENT(Statement_AddActivityData);
 		PREPARE_STATEMENT(Statement_SetActivityData);
@@ -1066,6 +1076,87 @@ public:
 		return Statement_GetPackageTransactionEventIdAtLastSave.BindAndExecuteSingle(InPackageNameId, OutTransactionEventId);
 	}
 
+	/**
+	 * Statements working on replication_events
+	 */
+
+	/** Set the endpoint data in endpoints for the given endpoint_id */
+	SQLITE_PREPARED_STATEMENT_BINDINGS_ONLY(FAddReplicationData, "INSERT INTO replication_events(activity_type, event_data_flags, event_data_size_bytes, event_data) VALUES(?1, ?2, ?3, ?4);", SQLITE_PREPARED_STATEMENT_BINDINGS(EConcertSyncReplicationActivityType, int32, int32, TArray<uint8>));
+	FAddReplicationData Statement_AddReplicationData;
+	bool AddReplicationData(const FConcertSyncReplicationEvent& InEventData, int64& OutActivityEventId)
+	{
+		const FConcertSessionSerializedPayload& Payload = InEventData.Payload;
+		const bool bIsCbor = Payload.SerializationMethod == EConcertPayloadSerializationMethod::Cbor;
+		// We don't save SerializationMethod in the database; instead we always assume that it's saved in Cbor format.
+		// If we allow the saving here, then GetReplicationEventForId will not deserialize correctly.
+		// If you want to change this, add a new field to the database.
+		if (!ensure(bIsCbor))
+		{
+			return false;
+		}
+		
+		const bool bSuccess = Statement_AddReplicationData.BindAndExecute(
+			InEventData.ActivityType,
+			PackageDataUtil::ConvertSerializedPayloadFlagsToInt32(Payload),
+			Payload.PayloadSize,
+			Payload.PayloadBytes.Bytes
+			);
+		if (bSuccess)
+		{
+			OutActivityEventId = Database.GetLastInsertRowId();
+		}
+		return bSuccess;
+	}
+
+	/** Set the endpoint data in endpoints for the given endpoint_id */
+	SQLITE_PREPARED_STATEMENT_BINDINGS_ONLY(FSetReplicationData, "INSERT OR REPLACE INTO replication_events(replication_event_id, activity_type, event_data_flags, event_data_size_bytes, event_data) VALUES(?1, ?2, ?3, ?4, ?5);", SQLITE_PREPARED_STATEMENT_BINDINGS(int64, EConcertSyncReplicationActivityType, int32, int32, TArray<uint8>));
+	FSetReplicationData Statement_SetReplicationData;
+	bool SetReplicationData(const int64 InActivityEventId, const FConcertSyncReplicationEvent& InEventData)
+	{
+		const FConcertSessionSerializedPayload& Payload = InEventData.Payload;
+		const bool bIsCbor = Payload.SerializationMethod == EConcertPayloadSerializationMethod::Cbor;
+		// We don't save SerializationMethod in the database; instead we always assume that it's saved in Cbor format.
+		// If we allow the saving here, then GetReplicationEventForId will not deserialize correctly.
+		// If you want to change this, add a new field to the database.
+		if (!ensure(bIsCbor))
+		{
+			return false;
+		}
+		
+		return Statement_SetReplicationData.BindAndExecute(
+			InActivityEventId,
+			InEventData.ActivityType,
+			PackageDataUtil::ConvertSerializedPayloadFlagsToInt32(Payload),
+			Payload.PayloadSize,
+			Payload.PayloadBytes.Bytes
+			);
+	}
+	
+	/** Set the endpoint data from endpoints for the given endpoint_id */
+	SQLITE_PREPARED_STATEMENT(FGetReplicationEventForId, "SELECT activity_type, event_data_flags, event_data_size_bytes, event_data FROM replication_events WHERE replication_event_id = ?1;", SQLITE_PREPARED_STATEMENT_COLUMNS(EConcertSyncReplicationActivityType, int32, int32, TArray<uint8>), SQLITE_PREPARED_STATEMENT_BINDINGS(int64));
+	FGetReplicationEventForId Statement_GetReplicationEventForId;
+	bool GetReplicationEventForId(const int64 InReplicationEventId, FConcertSyncReplicationEvent& OutReplicationEvent)
+	{
+		FConcertSessionSerializedPayload& Payload = OutReplicationEvent.Payload;
+		int32 SerializationFlags;
+		if (Statement_GetReplicationEventForId.BindAndExecuteSingle(InReplicationEventId, OutReplicationEvent.ActivityType, SerializationFlags, Payload.PayloadSize, Payload.PayloadBytes.Bytes))
+		{
+			PackageDataUtil::SetSerializedPayloadFlags(Payload, SerializationFlags);
+			Payload.PayloadTypeName = *UE::ConcertSyncCore::GetReplicationActivityPayloadTypePathName(OutReplicationEvent.ActivityType);
+			Payload.SerializationMethod = EConcertPayloadSerializationMethod::Cbor;
+			return true;
+		}
+		return false;
+	}
+	
+	/** Get the largest replication_event_id currently in replication_events that was triggered by EndpointId and has ActivityType */
+	SQLITE_PREPARED_STATEMENT(FGetReplicationMaxEventIdByClientAndType, "SELECT MAX(replication_event_id) FROM replication_events JOIN activities ON replication_events.replication_event_id = activities.event_id WHERE endpoint_id = ?1 AND activity_type = ?2;", SQLITE_PREPARED_STATEMENT_COLUMNS(int64), SQLITE_PREPARED_STATEMENT_BINDINGS(FGuid, EConcertSyncReplicationActivityType));
+	FGetReplicationMaxEventIdByClientAndType Statement_GetReplicationMaxEventIdByClientAndType;
+	bool GetReplicationMaxEventIdByClientAndType(const FGuid& EndpointId, EConcertSyncReplicationActivityType ActivityType, int64& OutReplicationEventId)
+	{
+		return Statement_GetReplicationMaxEventIdByClientAndType.BindAndExecuteSingle(EndpointId, ActivityType, OutReplicationEventId);
+	}
+
 	/**
 	 * Statements working on activities
 	 */
@@ -1599,6 +1690,7 @@ bool FConcertSyncSessionDatabase::Open(const FString& InSessionPath, const ESQLi
 	CREATE_TABLE("transaction_events", "transaction_event_id INTEGER PRIMARY KEY, data_filename TEXT NOT NULL");
 	CREATE_TABLE("package_events", "package_event_id INTEGER PRIMARY KEY, package_name_id INTEGER NOT NULL, package_revision INTEGER NOT NULL, package_info_flags INTEGER NOT NULL, package_info_size_bytes INTEGER NOT NULL, package_info_data BLOB, transaction_event_id_at_save INTEGER NOT NULL, data_filename TEXT NOT NULL, FOREIGN KEY(package_name_id) REFERENCES package_names(package_name_id)");
 	CREATE_TABLE("persist_events", "persist_event_id INTEGER PRIMARY KEY, package_event_id INTEGER NOT NULL, transaction_event_id_at_persist INTEGER NOT NULL, FOREIGN KEY(package_event_id) REFERENCES package_events(package_event_id)");
+	CREATE_TABLE("replication_events", "replication_event_id INTEGER PRIMARY KEY, activity_type INTEGER NOT NULL, event_data_flags INTEGER NOT NULL, event_data_size_bytes INTEGER NOT NULL, event_data BLOB");
 	CREATE_TABLE("activities", "activity_id INTEGER PRIMARY KEY, endpoint_id BLOB NOT NULL, event_time INTEGER NOT NULL, event_type INTEGER NOT NULL, event_id INTEGER NOT NULL, event_summary_type TEXT NOT NULL, event_summary_flags INTEGER NOT NULL, event_summary_size_bytes INTEGER NOT NULL, event_summary_data BLOB, activity_flags TINYINT UNSIGNED NOT NULL ,FOREIGN KEY(endpoint_id) REFERENCES endpoints(endpoint_id)");
 	CREATE_TABLE("ignored_activities", "activity_id INTEGER NOT NULL, FOREIGN KEY(activity_id) REFERENCES activities(activity_id)");
 	CREATE_TABLE("resource_locks", "object_name_id INTEGER NOT NULL, lock_event_id INTEGER NOT NULL, FOREIGN KEY(object_name_id) REFERENCES object_names(object_name_id), FOREIGN KEY(lock_event_id) REFERENCES lock_events(lock_event_id)");
@@ -1744,6 +1836,16 @@ bool FConcertSyncSessionDatabase::AddPackageActivity(const FConcertSyncActivity&
 		);
 }
 
+bool FConcertSyncSessionDatabase::AddReplicationActivity(const FConcertSyncReplicationActivity& InReplicationActivity, int64& OutActivityId, int64& OutReplicationEventId)
+{
+	FConcertSyncSessionDatabaseScopedTransaction ScopedTransaction(*Statements);
+	return ScopedTransaction.CommitOrRollback(
+		AddReplicationEvent(InReplicationActivity.EventData, OutReplicationEventId) &&
+		Statements->AddActivityData(InReplicationActivity.Flags, InReplicationActivity.EndpointId, EConcertSyncActivityEventType::Replication, OutReplicationEventId, InReplicationActivity.EventSummary, OutActivityId) &&
+		SetActivityIgnoredState(OutActivityId, InReplicationActivity.bIgnored)
+		);
+}
+
 bool FConcertSyncSessionDatabase::SetActivities(const TSet<int64>& ActivityIds, TFunctionRef<void(FConcertSyncActivity&)> UpdateCallback)
 {
 	FConcertSyncSessionDatabaseScopedTransaction ScopedTransaction(*Statements);
@@ -1813,6 +1915,16 @@ bool FConcertSyncSessionDatabase::SetPackageActivity(const FConcertSyncActivity&
 		);
 }
 
+bool FConcertSyncSessionDatabase::SetReplicationActivity(const FConcertSyncReplicationActivity& InReplicationActivity)
+{
+	FConcertSyncSessionDatabaseScopedTransaction ScopedTransaction(*Statements);
+	return ScopedTransaction.CommitOrRollback(
+		SetReplicationEvent(InReplicationActivity.EventId, InReplicationActivity.EventData) &&
+		Statements->SetActivityData(InReplicationActivity.ActivityId, InReplicationActivity.Flags, InReplicationActivity.EndpointId, InReplicationActivity.EventTime, InReplicationActivity.EventType, InReplicationActivity.EventId, InReplicationActivity.EventSummary) &&
+		SetActivityIgnoredState(InReplicationActivity.ActivityId, InReplicationActivity.bIgnored)
+		);
+}
+
 bool FConcertSyncSessionDatabase::GetActivity(const int64 InActivityId, FConcertSyncActivity& OutActivity) const
 {
 	OutActivity.ActivityId = InActivityId;
@@ -1851,6 +1963,12 @@ bool FConcertSyncSessionDatabase::GetPackageActivity(const int64 InActivityId, F
 		&& GetPackageEvent(PackageActivityBasePart.EventId, GetPackageActivityEventFn); // Pull the event part of the activity and call back (to let the caller stream the package data).
 }
 
+bool FConcertSyncSessionDatabase::GetReplicationActivity(const int64 InActivityId, FConcertSyncReplicationActivity& OutReplicationActivity) const
+{
+	return GetActivity(InActivityId, OutReplicationActivity)
+		&& GetReplicationEvent(OutReplicationActivity.EventId, OutReplicationActivity.EventData);
+}
+
 bool FConcertSyncSessionDatabase::GetActivityEventType(const int64 InActivityId, EConcertSyncActivityEventType& OutEventType) const
 {
 	return Statements->GetActivityEventTypeForId(InActivityId, OutEventType);
@@ -2012,6 +2130,29 @@ bool FConcertSyncSessionDatabase::EnumeratePackageActivities(FIteratePackageActi
 	});
 }
 
+bool FConcertSyncSessionDatabase::EnumerateReplicationActivities(TFunctionRef<bool(FConcertSyncReplicationActivity&&)> InCallback) const
+{
+	return Statements->GetAllActivityDataForEventType(EConcertSyncActivityEventType::Transaction, [this, &InCallback](const int64 InActivityId, const EConcertSyncActivityFlags InFlags, const FGuid& InEndpointId, const FDateTime InEventTime, const int64 InEventId, FConcertSessionSerializedPayload&& InEventSummary)
+	{
+		FConcertSyncReplicationActivity ReplicationActivity;
+		ReplicationActivity.ActivityId = InActivityId;
+		ReplicationActivity.bIgnored = Statements->IsActivityIgnored(InActivityId);
+		ReplicationActivity.Flags = InFlags;
+		ReplicationActivity.EndpointId = InEndpointId;
+		ReplicationActivity.EventTime = InEventTime;
+		ReplicationActivity.EventType = EConcertSyncActivityEventType::Transaction;
+		ReplicationActivity.EventId = InEventId;
+		ReplicationActivity.EventSummary = MoveTemp(InEventSummary);
+		if (GetReplicationEvent(ReplicationActivity.EventId, ReplicationActivity.EventData))
+		{
+			return InCallback(MoveTemp(ReplicationActivity))
+				? ESQLitePreparedStatementExecuteRowResult::Continue
+				: ESQLitePreparedStatementExecuteRowResult::Stop;
+		}
+		return ESQLitePreparedStatementExecuteRowResult::Error;
+	});
+}
+
 bool FConcertSyncSessionDatabase::EnumerateActivitiesForEventType(const EConcertSyncActivityEventType InEventType, FIterateActivityFunc InCallback) const
 {
 	return Statements->GetAllActivityDataForEventType(InEventType, [this, InEventType, &InCallback](const int64 InActivityId, const EConcertSyncActivityFlags InFlags, const FGuid& InEndpointId, const FDateTime InEventTime, const int64 InEventId, FConcertSessionSerializedPayload&& InEventSummary)
@@ -2247,6 +2388,11 @@ bool FConcertSyncSessionDatabase::GetTransactionEvent(const int64 InTransactionE
 	return false;
 }
 
+bool FConcertSyncSessionDatabase::GetReplicationEvent(const int64 InReplicationEventId, FConcertSyncReplicationEvent& OutReplicationEvent) const
+{
+	return Statements->GetReplicationEventForId(InReplicationEventId, OutReplicationEvent);
+}
+
 bool FConcertSyncSessionDatabase::GetTransactionMaxEventId(int64& OutTransactionEventId) const
 {
 	return Statements->GetTransactionMaxEventId(OutTransactionEventId);
@@ -2480,6 +2626,16 @@ bool FConcertSyncSessionDatabase::SetPackageEvent(const int64 InPackageEventId,
 	return Statements->SetPackageEvent(InPackageEventId, PackageNameId, InPackageRevision, InPackageInfo.TransactionEventIdAtSave, InPackageInfo, PackageDataFilename);
 }
 
+bool FConcertSyncSessionDatabase::AddReplicationEvent(const FConcertSyncReplicationEvent& InReplicationEvent, int64& OutReplicationEventId)
+{
+	return Statements->AddReplicationData(InReplicationEvent, OutReplicationEventId);
+}
+
+bool FConcertSyncSessionDatabase::SetReplicationEvent(const int64 InReplicationEventId, const FConcertSyncReplicationEvent& InReplicationEvent)
+{
+	return Statements->SetReplicationData(InReplicationEventId, InReplicationEvent);
+}
+
 bool FConcertSyncSessionDatabase::GetPackageEventMetaData(const int64 InPackageEventId, int64& OuptPackageRevision, FConcertPackageInfo& OutPackageInfo) const
 {
 	FString UnusedDataFilename;
@@ -2666,6 +2822,11 @@ bool FConcertSyncSessionDatabase::IsHeadRevisionPackageEvent(const int64 InPacka
 	return false;
 }
 
+bool FConcertSyncSessionDatabase::GetReplicationMaxEventIdByClientAndType(const FGuid& EndpointId, EConcertSyncReplicationActivityType ActivityType, int64& OutReplicationEventId) const
+{
+	return Statements->GetReplicationMaxEventIdByClientAndType(EndpointId, ActivityType, OutReplicationEventId);
+}
+
 bool FConcertSyncSessionDatabase::GetObjectPathName(const int64 InObjectNameId, FName& OutObjectPathName) const
 {
 	return Statements->GetObjectPathName(InObjectNameId, OutObjectPathName);
diff --git a/Engine/Plugins/Developer/Concert/ConcertSync/ConcertSyncCore/Source/ConcertSyncCore/Public/ConcertSyncSessionDatabase.h b/Engine/Plugins/Developer/Concert/ConcertSync/ConcertSyncCore/Source/ConcertSyncCore/Public/ConcertSyncSessionDatabase.h
--- a/Engine/Plugins/Developer/Concert/ConcertSync/ConcertSyncCore/Source/ConcertSyncCore/Public/ConcertSyncSessionDatabase.h
+++ b/Engine/Plugins/Developer/Concert/ConcertSync/ConcertSyncCore/Source/ConcertSyncCore/Public/ConcertSyncSessionDatabase.h
@@ -9,6 +9,7 @@
 class FStructOnScope;
 enum class EConcertSyncActivityEventType : uint8;
 enum class EConcertSyncActivityFlags : uint8;
+enum class EConcertSyncReplicationActivityType : uint8;
 struct FConcertPackageDataStream;
 struct FConcertPackageInfo;
 struct FConcertSessionFilter;
@@ -20,6 +21,8 @@ struct FConcertSyncEndpointIdAndData;
 struct FConcertSyncLockActivity;
 struct FConcertSyncLockEvent;
 struct FConcertSyncPackageEventData;
+struct FConcertSyncReplicationActivity;
+struct FConcertSyncReplicationEvent;
 struct FConcertSyncTransactionActivity;
 struct FConcertSyncTransactionEvent;
 struct FConcertTransactionFinalizedEvent;
@@ -152,6 +155,19 @@ public:
 	 */
 	bool AddPackageActivity(const FConcertSyncActivity& InPackageActivity, const FConcertPackageInfo& InPackageInfo, FConcertPackageDataStream& InPackageDataStream, int64& OutActivityId, int64& OutPackageEventId);
 
+	/**
+	 * Add a new replication activity to this database, assigning it both an activity and lock event ID.
+	 * @note The endpoint ID referenced by the activity must exist in the database (@see SetEndpoint).
+	 * @note This function is expected to be called on the server to populate its version of the session database.
+	 *
+	 * @param InReplicationActivity		The lock activity to add (the ActivityId, EventTime, EventType, and EventId members are ignored).
+	 * @param OutActivityId				Populated with the ID of this activity in the database.
+	 * @param OutReplicationEventId			Populated with the ID of the lock event in the database (@see GetLockEvent).
+	 *
+	 * @return True if the lock activity was added, false otherwise.
+	 */
+	bool AddReplicationActivity(const FConcertSyncReplicationActivity& InReplicationActivity, int64& OutActivityId, int64& OutReplicationEventId);
+	
 	/**
 	 * Iterates the given range, calls UpdateCallback on each element, and commits the update.
 	 */
@@ -204,6 +220,17 @@ public:
 	 */
 	bool SetPackageActivity(const FConcertSyncActivity& InPackageActivity, FConcertSyncPackageEventData& InPackageActivityEvent, const bool bMetaDataOnly = false);
 
+	/**
+	 * Set a replication activity in this database, creating or replacing it.
+	 * @note The endpoint ID referenced by the activity must exist in the database (@see SetEndpoint).
+	 * @note This function is expected to be called on the client to populate its version of the session database from data synced from the server.
+	 *
+	 * @param InReplicationActivity		The replication activity to set.
+	 *
+	 * @return True if the transaction activity was set, false otherwise.
+	 */
+	bool SetReplicationActivity(const FConcertSyncReplicationActivity& InReplicationActivity);
+	
 	/**
 	 * Get the generic part of an activity from this database.
 	 *
@@ -255,6 +282,16 @@ public:
 	 */
 	bool GetPackageActivity(const int64 InActivityId, FConsumePackageActivityFunc PackageActivityFn) const;
 
+	/**
+	 * Get a replication activity from this database.
+	 *
+	 * @param InActivityId				The ID of the activity to find.
+	 * @param OutReplicationActivity	The replication activity to populate with the result.
+	 *
+	 * @return True if the replication activity was found, false otherwise.
+	 */
+	bool GetReplicationActivity(const int64 InActivityId, FConcertSyncReplicationActivity& OutReplicationActivity) const;
+	
 	/**
 	 * Get the type of an activity in this database.
 	 *
@@ -362,6 +399,15 @@ public:
 	 */
 	bool EnumeratePackageActivities(FIteratePackageActivityFunc InCallback) const;
 
+	/**
+	 * Enumerate all the replication activities in this database.
+	 *
+	 * @param InCallback				Callback invoked for each activity; return true to continue enumeration, or false to stop.
+	 *
+	 * @return True if the replication activities were enumerated without error, false otherwise.
+	 */
+	bool EnumerateReplicationActivities(TFunctionRef<bool(FConcertSyncReplicationActivity&&)> InCallback) const;
+
 	/**
 	 * Enumerate all the activities in this database of the given type.
 	 *
@@ -492,6 +538,16 @@ public:
 	 */
 	bool GetTransactionEvent(const int64 InTransactionEventId, FConcertSyncTransactionEvent& OutTransactionEvent, const bool InMetaDataOnly = false) const;
 
+	/**
+	 * Get a replication event from this database.
+	 *
+	 * @param InReplicationEventId		The ID of the replication event to find.
+	 * @param OutReplicationEvent		The replication event to populate with the result.
+	 *
+	 * @return True if the replication event was found, false otherwise.
+	 */
+	bool GetReplicationEvent(const int64 InReplicationEventId, FConcertSyncReplicationEvent& OutReplicationEvent) const;
+
 	/**
 	 * Get the maximum ID of the transaction events in this database.
 	 *
@@ -672,6 +728,17 @@ public:
 	 */
 	bool IsHeadRevisionPackageEvent(const int64 InPackageEventId, bool& OutIsHeadRevision) const;
 
+	/**
+	 * Get the maximum ID of the replication events in this database that was triggered by EndpointId and had ActivityType.
+	 *
+	 * @param EndpointId				The ID of the client that created the activity
+	 * @param ActivityType				The activity type the replication activity should have 
+	 * @param OutReplicationEventId		The replication event ID to populate with the result.
+	 *
+	 * @return True if the replication event ID was resolved, false otherwise.
+	 */
+	bool GetReplicationMaxEventIdByClientAndType(const FGuid& EndpointId, EConcertSyncReplicationActivityType ActivityType, int64& OutReplicationEventId) const;
+
 	/**
 	 * Add a package event ID for the head revision to the persist events in this database, if not already existing.
 	 *
@@ -840,6 +907,26 @@ private:
 	 */
 	bool SetPackageEvent(const int64 InPackageEventId, int64 PackageRevision, const FConcertPackageInfo& PackageInfo, FConcertPackageDataStream* InPackageDataStream);
 
+	/**
+	 * Add a new replication event to this database, assigning it a replication event ID.
+	 *
+	 * @param InReplicationEvent		The transaction event to add.
+	 * @param OutReplicationEventId		Populated with the ID of the transaction event in the database.
+	 *
+	 * @return True if the replication event was added, false otherwise.
+	 */
+	bool AddReplicationEvent(const FConcertSyncReplicationEvent& InReplicationEvent, int64& OutReplicationEventId);
+
+	/**
+	 * Set a replication event in this database, creating or replacing it.
+	 *
+	 * @param InReplicationEventId		The ID of the replication event to set.
+	 * @param InReplicationEvent		The replication event to set.
+	 *
+	 * @return True if the replication event was set, false otherwise.
+	 */
+	bool SetReplicationEvent(const int64 InReplicationEventId, const FConcertSyncReplicationEvent& InReplicationEvent);
+	
 	/**
 	 * Get the maximum ID of the package events in this database.
 	 *
diff --git a/Engine/Plugins/Developer/Concert/ConcertSync/ConcertSyncCore/Source/ConcertSyncCore/Public/ConcertSyncSessionTypes.h b/Engine/Plugins/Developer/Concert/ConcertSync/ConcertSyncCore/Source/ConcertSyncCore/Public/ConcertSyncSessionTypes.h
--- a/Engine/Plugins/Developer/Concert/ConcertSync/ConcertSyncCore/Source/ConcertSyncCore/Public/ConcertSyncSessionTypes.h
+++ b/Engine/Plugins/Developer/Concert/ConcertSync/ConcertSyncCore/Source/ConcertSyncCore/Public/ConcertSyncSessionTypes.h
@@ -32,6 +32,10 @@ enum class EConcertSyncActivityEventType : uint8
 	Lock,
 	Transaction,
 	Package,
+	Replication,
+
+	// ADD NEW ENTRIES ABOVE 
+	Count
 };
 
 /** Type of transaction summaries */
diff --git a/Engine/Plugins/Developer/Concert/ConcertSync/ConcertSyncCore/Source/ConcertSyncCore/Public/Replication/Messages/ReplicationActivity.h b/Engine/Plugins/Developer/Concert/ConcertSync/ConcertSyncCore/Source/ConcertSyncCore/Public/Replication/Messages/ReplicationActivity.h
new file mode 100644
--- /dev/null
+++ b/Engine/Plugins/Developer/Concert/ConcertSync/ConcertSyncCore/Source/ConcertSyncCore/Public/Replication/Messages/ReplicationActivity.h
@@ -0,0 +1,142 @@
+﻿// Copyright Epic Games, Inc. All Rights Reserved.
+
+#pragma once
+
+#include "ConcertSyncSessionTypes.h"
+#include "Containers/Array.h"
+#include "Containers/UnrealString.h"
+#include "Replication/Data/ObjectIds.h"
+#include "Replication/Data/ReplicationStream.h"
+#include "ReplicationActivity.generated.h"
+
+/** Identifies the FConcertSyncReplicationEvent::Payload struct type */
+UENUM()
+enum class EConcertSyncReplicationActivityType : uint8
+{
+	None,
+	
+	/** Client left the replication session (either because they left the MU session or because they sent a FConcertReplication_LeaveEvent). */
+	LeaveReplication,
+
+	// ADD NEW ENTRIES ABOVE
+	Count
+};
+
+/**
+ * Contains the streams and authority a client had when they left a session.
+ * Used for restoring content when the client rejoins.
+ */
+USTRUCT()
+struct FConcertSyncReplicationPayload_LeaveReplication
+{
+	GENERATED_BODY()
+
+	/** The streams the client had registered when they left. */
+	UPROPERTY()
+	TArray<FConcertReplicationStream> Streams;
+	
+	/** The objects the client had authority over when they left. */
+	UPROPERTY()
+	TArray<FConcertObjectInStreamID> OwnedObjects;
+
+	friend bool operator==(const FConcertSyncReplicationPayload_LeaveReplication&, const FConcertSyncReplicationPayload_LeaveReplication&) = default;
+	friend bool operator!=(const FConcertSyncReplicationPayload_LeaveReplication&, const FConcertSyncReplicationPayload_LeaveReplication&) = default;
+};
+
+namespace UE::ConcertSyncCore
+{
+	inline FString GetReplicationActivityPayloadTypePathName(EConcertSyncReplicationActivityType Type)
+	{
+		switch (Type)
+		{
+		case EConcertSyncReplicationActivityType::LeaveReplication: return FConcertSyncReplicationPayload_LeaveReplication::StaticStruct()->GetPathName();
+		default: ensureMsgf(false, TEXT("Unknown replication activity type (%u)"), static_cast<uint8>(Type)); return TEXT("");
+		}
+	}
+}
+
+/** Data for a replication event in a Concert Sync Session */
+USTRUCT()
+struct FConcertSyncReplicationEvent
+{
+	GENERATED_BODY()
+
+	/** Identifies the Payload struct type */
+	UPROPERTY()
+	EConcertSyncReplicationActivityType ActivityType = EConcertSyncReplicationActivityType::None;
+
+	/**
+	 * A FConcertSyncReplicationPayload_X type depending on ActivityType.
+	 * 
+	 * Serialized into the database using Cbor; do not change (for simplicity we always assume it is Cbor).
+	 * If you change the serialization method, you'll get a failing check.
+	 * If for whatever reason you MUST save it using a different method, adjust FConcertSyncSessionDatabaseStatements::AddReplicationData
+	 * and FConcertSyncSessionDatabaseStatements::SetReplicationData.
+	 */
+	UPROPERTY()
+	FConcertSessionSerializedPayload Payload{ EConcertPayloadSerializationMethod::Cbor };
+
+	FConcertSyncReplicationEvent() = default;
+	FConcertSyncReplicationEvent(const FConcertSyncReplicationPayload_LeaveReplication& Data)
+	{
+		SetPayload(Data);
+	}
+
+	void SetPayload(const FConcertSyncReplicationPayload_LeaveReplication& Data)
+	{
+		ActivityType = EConcertSyncReplicationActivityType::LeaveReplication;
+		Payload.SetTypedPayload(Data);
+	}
+	
+	bool GetPayload(FConcertSyncReplicationPayload_LeaveReplication& Result) const
+	{
+		check(ActivityType == EConcertSyncReplicationActivityType::LeaveReplication);
+		return Payload.GetTypedPayload(Result);
+	}
+
+	friend bool operator==(const FConcertSyncReplicationEvent& Left, const FConcertSyncReplicationEvent& Right)
+	{
+		if (Left.ActivityType != Right.ActivityType)
+		{
+			return false;
+		}
+
+		static_assert(static_cast<int32>(EConcertSyncReplicationActivityType::Count) == 2, "If you added an EConcertSyncReplicationActivityType entry, update this switch");
+		switch (Left.ActivityType)
+		{
+		case EConcertSyncReplicationActivityType::None: return true;
+		case EConcertSyncReplicationActivityType::LeaveReplication:
+		{
+			FConcertSyncReplicationPayload_LeaveReplication LeftContent;
+			FConcertSyncReplicationPayload_LeaveReplication RightContent;
+			const bool bLeftSucceeded = Left.GetPayload(LeftContent);
+			const bool bRightSucceeded = Right.GetPayload(RightContent);
+			return bLeftSucceeded && bRightSucceeded && LeftContent == RightContent;
+		}
+		default: checkNoEntry(); return false;
+		}
+	}
+	friend bool operator!=(const FConcertSyncReplicationEvent&, const FConcertSyncReplicationEvent&) = default;
+};
+
+/** Data for a replication activity entry in a Concert Sync Session */
+USTRUCT()
+struct FConcertSyncReplicationActivity : public FConcertSyncActivity
+{
+	GENERATED_BODY()
+
+	FConcertSyncReplicationActivity()
+	{
+		EventType = EConcertSyncActivityEventType::Replication;
+	}
+
+	FConcertSyncReplicationActivity(const FConcertSyncReplicationPayload_LeaveReplication& PayloadData)
+		: EventData(PayloadData)
+	{
+		EventType = EConcertSyncActivityEventType::Replication;
+	}
+
+	/** The transaction event data associated with this activity */
+	UPROPERTY()
+	FConcertSyncReplicationEvent EventData;
+};
\ No newline at end of file
diff --git a/Engine/Plugins/Developer/Concert/ConcertSync/ConcertSyncServer/Source/ConcertSyncServer/Private/ConcertServerWorkspace.cpp b/Engine/Plugins/Developer/Concert/ConcertSync/ConcertSyncServer/Source/ConcertSyncServer/Private/ConcertServerWorkspace.cpp
--- a/Engine/Plugins/Developer/Concert/ConcertSync/ConcertSyncServer/Source/ConcertSyncServer/Private/ConcertServerWorkspace.cpp
+++ b/Engine/Plugins/Developer/Concert/ConcertSync/ConcertSyncServer/Source/ConcertSyncServer/Private/ConcertServerWorkspace.cpp
@@ -13,9 +13,11 @@
 #include "ConcertLogGlobal.h"
 #include "ConcertPackageEvents.h"
 #include "ConcertUtil.h"
+#include "HistoryEdition/ActivityDependencyGraph.h"
+#include "Replication/Messages/ReplicationActivity.h"
 #include "Serialization/MemoryReader.h"
+
 #include "Algo/Transform.h"
-#include "HistoryEdition/ActivityDependencyGraph.h"
 
 FConcertServerWorkspace::FConcertServerWorkspace(const TSharedRef<FConcertSyncServerLiveSession>& InLiveSession, TSharedPtr<IConcertFileSharingService> InFileSharingService)
 	: FileSharingService(MoveTemp(InFileSharingService))
@@ -1150,6 +1152,41 @@ void FConcertServerWorkspace::SendSyncPackageActivityEvent(const FConcertWorkspa
 	LiveSession->GetSession().SendCustomEvent(SyncEvent, InTargetEndpointId, EConcertMessageFlags::ReliableOrdered);
 }
 
+void FConcertServerWorkspace::AddReplicationActivity(const FConcertSyncReplicationActivity& InReplicationActivity)
+{
+	// Add the activity and sync it
+	int64 ActivityId = 0;
+	int64 EventId = 0;
+	if (LiveSession->GetSessionDatabase().AddReplicationActivity(InReplicationActivity, ActivityId, EventId))
+	{
+		PostActivityAdded(ActivityId);
+		SyncCommandQueue->QueueCommand(LiveSyncEndpoints, [this, SyncActivityId = ActivityId](const FConcertServerSyncCommandQueue::FSyncCommandContext& InSyncCommandContext, const FGuid& InEndpointId)
+		{
+			SendSyncReplicationActivityEvent(InEndpointId, SyncActivityId, InSyncCommandContext.GetNumRemainingCommands());
+		});
+	}
+	else
+	{
+		UE_LOG(LogConcert, Error, TEXT("Failed to set replication activity '%s' on live session '%s': %s"), *LexToString(ActivityId), *LiveSession->GetSession().GetName(), *LiveSession->GetSessionDatabase().GetLastError());
+	}	
+}
+
+void FConcertServerWorkspace::SendSyncReplicationActivityEvent(const FGuid& InTargetEndpointId, const int64 InSyncActivityId, const int32 InNumRemainingSyncEvents)
+{
+	FConcertSyncReplicationActivity SyncActivity;
+	if (LiveSession->GetSessionDatabase().GetReplicationActivity(InSyncActivityId, SyncActivity))
+	{
+		FConcertWorkspaceSyncActivityEvent SyncEvent;
+		SyncEvent.NumRemainingSyncEvents = InNumRemainingSyncEvents;
+		SyncEvent.Activity.SetTypedPayload(SyncActivity);
+		LiveSession->GetSession().SendCustomEvent(SyncEvent, InTargetEndpointId, EConcertMessageFlags::ReliableOrdered);
+	}
+	else
+	{
+		UE_LOG(LogConcert, Error, TEXT("Failed to get replication activity '%s' from live session '%s': %s"), *LexToString(InSyncActivityId), *LiveSession->GetSession().GetName(), *LiveSession->GetSessionDatabase().GetLastError());
+	}
+}
+
 void FConcertServerWorkspace::PostActivityAdded(const int64 InActivityId)
 {
 	FConcertSyncActivity Activity;
diff --git a/Engine/Plugins/Developer/Concert/ConcertSync/ConcertSyncServer/Source/ConcertSyncServer/Private/ConcertServerWorkspace.h b/Engine/Plugins/Developer/Concert/ConcertSync/ConcertSyncServer/Source/ConcertSyncServer/Private/ConcertServerWorkspace.h
--- a/Engine/Plugins/Developer/Concert/ConcertSync/ConcertSyncServer/Source/ConcertSyncServer/Private/ConcertServerWorkspace.h
+++ b/Engine/Plugins/Developer/Concert/ConcertSync/ConcertSyncServer/Source/ConcertSyncServer/Private/ConcertServerWorkspace.h
@@ -16,6 +16,7 @@ class FConcertServerDataStore;
 
 struct FConcertMessageContext;
 struct FConcertEndpointContext;
+struct FConcertSyncReplicationActivity;
 struct FConcertTransactionSnapshotEvent;
 struct FConcertTransactionFinalizedEvent;
 
@@ -249,6 +250,24 @@ private:
 	 * @param InSyncActivityId		The ID of the activity to send the sync event for.
 	 */
 	void SendSyncPackageActivityEvent(const FConcertWorkspaceSyncActivityEvent& SyncEvent, const FGuid& InTargetEndpointId) const;
+
+	/**
+	 * Add a new replication activity to the session database, and sync the result back to all clients.
+	 * @note The endpoint ID referenced by the activity must exist in the database (@see SetEndpoint).
+	 *
+	 * @param InReplicationActivity		The replication activity to add (the ActivityId, EventTime, EventType, and EventId members are ignored).
+	 */
+	void AddReplicationActivity(const FConcertSyncReplicationActivity& InReplicationActivity);
+	
+	/**
+	 * Send a sync event for a replication activity in the session database.
+	 *
+	 * @param InTargetEndpointId		The ID of the endpoint to send the sync event to.
+	 * @param InSyncActivityId			The ID of the activity to send the sync event for.
+	 * @param InNumRemainingSyncEvents	The number of items left in the sync queue.
+	 */
+	void SendSyncReplicationActivityEvent(const FGuid& InTargetEndpointId, const int64 InSyncActivityId, const int32 InNumRemainingSyncEvents);
+	
 	/**
 	 * Called after any activity is added to the session database.
 	 *
diff --git a/Engine/Plugins/Developer/Concert/ConcertSync/ConcertSyncServer/Source/ConcertSyncServer/Private/ConcertSyncServer.cpp b/Engine/Plugins/Developer/Concert/ConcertSync/ConcertSyncServer/Source/ConcertSyncServer/Private/ConcertSyncServer.cpp
--- a/Engine/Plugins/Developer/Concert/ConcertSync/ConcertSyncServer/Source/ConcertSyncServer/Private/ConcertSyncServer.cpp
+++ b/Engine/Plugins/Developer/Concert/ConcertSync/ConcertSyncServer/Source/ConcertSyncServer/Private/ConcertSyncServer.cpp
@@ -11,6 +11,7 @@
 #include "ConcertSyncSessionDatabase.h"
 #include "ConcertLogGlobal.h"
 #include "Replication/ConcertServerReplicationManager.h"
+#include "Replication/Messages/ReplicationActivity.h"
 
 #include "Misc/Paths.h"
 #include "HAL/FileManager.h"
@@ -108,6 +109,7 @@ bool MigrateSessionData(const FConcertSyncSessionDatabase& InSourceDatabase, con
 				return true;
 			}
 
+			static_assert(static_cast<int32>(EConcertSyncActivityEventType::Count) == 6, "If you added an EConcertSyncActivityEventType entry, update this switch");
 			switch (InEventType)
 			{
 			case EConcertSyncActivityEventType::Connection:
@@ -209,6 +211,23 @@ bool MigrateSessionData(const FConcertSyncSessionDatabase& InSourceDatabase, con
 			}
 			break;
 
+			case EConcertSyncActivityEventType::Replication:
+			{
+				FConcertSyncReplicationActivity ReplicationActivity;
+				if (!InSourceDatabase.GetReplicationActivity(InActivityId, ReplicationActivity))
+				{
+					MIGRATE_SET_ERROR_RESULT_AND_RETURN("Failed to get replication activity '%s' from database at '%s': %s", *LexToString(InActivityId), *InSourceDatabase.GetFilename(), *InSourceDatabase.GetLastError());
+				}
+				if (InDestSessionFilter.bIncludeIgnoredActivities || !ReplicationActivity.bIgnored)
+				{
+					if (!DestDatabase.SetReplicationActivity(ReplicationActivity))
+					{
+						MIGRATE_SET_ERROR_RESULT_AND_RETURN("Failed to set replication activity '%s' on database at '%s': %s", *LexToString(InActivityId), *DestDatabase.GetFilename(), *DestDatabase.GetLastError());
+					}
+				}
+			}
+			break;
+
 			default:
 				checkf(false, TEXT("Unhandled EConcertSyncActivityEventType when migrating session database"));
 				break;
diff --git a/Engine/Plugins/Developer/Concert/ConcertSync/ConcertSyncTest/Source/ConcertSyncTest/Private/ConcertDatabaseTests.cpp b/Engine/Plugins/Developer/Concert/ConcertSync/ConcertSyncTest/Source/ConcertSyncTest/Private/ConcertDatabaseTests.cpp
deleted file mode 100644
diff --git a/Engine/Plugins/Developer/Concert/ConcertSync/ConcertSyncTest/Source/ConcertSyncTest/Private/Database/ConcertDatabaseTests.cpp b/Engine/Plugins/Developer/Concert/ConcertSync/ConcertSyncTest/Source/ConcertSyncTest/Private/Database/ConcertDatabaseTests.cpp
new file mode 100644
--- /dev/null
+++ b/Engine/Plugins/Developer/Concert/ConcertSync/ConcertSyncTest/Source/ConcertSyncTest/Private/Database/ConcertDatabaseTests.cpp
@@ -0,0 +1,290 @@
+// Copyright Epic Games, Inc. All Rights Reserved.
+
+#include "ConcertSyncSessionDatabase.h"
+#include "ConcertSyncSessionTypes.h"
+#include "Replication/Data/ConcertPropertySelection.h"
+#include "Replication/Data/ObjectIds.h"
+#include "Replication/Data/ReplicationStream.h"
+#include "Replication/Messages/ReplicationActivity.h"
+
+#include "Components/StaticMeshComponent.h"
+#include "Misc/AutomationTest.h"
+#include "Misc/Paths.h"
+#include "HAL/FileManager.h"
+
+IMPLEMENT_SIMPLE_AUTOMATION_TEST(FConcertSessionDatabaseTest, "Editor.Concert.Database.SessionDatabase", EAutomationTestFlags::EditorContext | EAutomationTestFlags::EngineFilter)
+
+bool FConcertSessionDatabaseTest::RunTest(const FString& Parameters)
+{
+	const FString TestSessionPath_Server = FPaths::ProjectIntermediateDir() / TEXT("ConcertDatabaseTest_Server");
+	const FString TestSessionPath_Client = FPaths::ProjectIntermediateDir() / TEXT("ConcertDatabaseTest_Client");
+
+	FConcertSyncSessionDatabase SessionDatabase_Server;
+	FConcertSyncSessionDatabase SessionDatabase_Client;
+
+	if (!SessionDatabase_Server.Open(TestSessionPath_Server))
+	{
+		AddError(TEXT("Failed to open server database"));
+	}
+
+	if (!SessionDatabase_Client.Open(TestSessionPath_Client))
+	{
+		AddError(TEXT("Failed to open client database"));
+	}
+
+	// Populate the server database with various activities
+	if (SessionDatabase_Server.IsValid())
+	{
+		// Add the endpoint the activity will be associated with
+		const FGuid EndpointId = FGuid::NewGuid();
+		{
+			FConcertSyncEndpointData EndpointData;
+			EndpointData.ClientInfo.Initialize();
+			if (!SessionDatabase_Server.SetEndpoint(EndpointId, EndpointData))
+			{
+				AddError(FString::Printf(TEXT("Failed to set endpoint on server database: %s"), *SessionDatabase_Server.GetLastError()));
+			}
+		}
+
+		// Add some activity that would have been generated by a client
+		{
+			int64 ActivityId = 0;
+			int64 EventId = 0;
+
+			FConcertSyncConnectionActivity ConnectionActivity;
+			ConnectionActivity.EndpointId = EndpointId;
+			ConnectionActivity.EventData.ConnectionEventType = EConcertSyncConnectionEventType::Connected;
+			if (!SessionDatabase_Server.AddConnectionActivity(ConnectionActivity, ActivityId, EventId))
+			{
+				AddError(FString::Printf(TEXT("Failed to add connection activity to server database: %s"), *SessionDatabase_Server.GetLastError()));
+			}
+
+			FConcertSyncLockActivity LockActivity;
+			LockActivity.EndpointId = EndpointId;
+			LockActivity.EventData.ResourceNames.Add(TEXT("/Game/TestAsset"));
+			LockActivity.EventData.LockEventType = EConcertSyncLockEventType::Locked;
+			if (!SessionDatabase_Server.AddLockActivity(LockActivity, ActivityId, EventId))
+			{
+				AddError(FString::Printf(TEXT("Failed to add lock activity to server database: %s"), *SessionDatabase_Server.GetLastError()));
+			}
+			LockActivity.EventData.LockEventType = EConcertSyncLockEventType::Unlocked;
+			if (!SessionDatabase_Server.AddLockActivity(LockActivity, ActivityId, EventId))
+			{
+				AddError(FString::Printf(TEXT("Failed to add lock activity to server database: %s"), *SessionDatabase_Server.GetLastError()));
+			}
+
+			FConcertSyncTransactionActivity TransactionActivity;
+			TransactionActivity.EndpointId = EndpointId;
+			TransactionActivity.EventData.Transaction.ModifiedPackages.Add(TEXT("/Game/TestAsset"));
+			if (!SessionDatabase_Server.AddTransactionActivity(TransactionActivity, ActivityId, EventId))
+			{
+				AddError(FString::Printf(TEXT("Failed to add transaction activity to server database: %s"), *SessionDatabase_Server.GetLastError()));
+			}
+
+			FConcertSyncActivity PackageActivityBasePart;
+			FConcertPackageInfo PackageInfo;
+			FConcertPackageDataStream PackageDataStream;
+			PackageActivityBasePart.EndpointId = EndpointId;
+			PackageInfo.PackageName = TEXT("/Game/TestAsset");
+			SessionDatabase_Server.GetTransactionMaxEventId(PackageInfo.TransactionEventIdAtSave);
+			if (!SessionDatabase_Server.AddPackageActivity(PackageActivityBasePart, PackageInfo, PackageDataStream, ActivityId, EventId))
+			{
+				AddError(FString::Printf(TEXT("Failed to add package activity to server database: %s"), *SessionDatabase_Server.GetLastError()));
+			}
+
+			const FGuid StreamId = FGuid::NewGuid();
+			const FSoftObjectPath StaticMeshComponent(TEXT("/Game/Maps.Map:PersistentLevel.Cube.StaticMeshComponent0"));
+			FConcertReplicationStream Stream;
+			Stream.BaseDescription.Identifier = StreamId;
+			FConcertReplicatedObjectInfo& ObjectInfo = Stream.BaseDescription.ReplicationMap.ReplicatedObjects.Add(StaticMeshComponent);
+			ObjectInfo.ClassPath = UStaticMeshComponent::StaticClass();
+			const TOptional<FConcertPropertyChain> Property_RelativeX = FConcertPropertyChain::CreateFromPath(*UStaticMeshComponent::StaticClass(), { TEXT("RelativeLocation"), TEXT("X") });
+			if (ensure(Property_RelativeX))
+			{
+				ObjectInfo.PropertySelection.ReplicatedProperties.Add(*Property_RelativeX);
+			}
+			FConcertSyncReplicationPayload_LeaveReplication LeaveReplication;
+			LeaveReplication.Streams.Add(Stream);
+			LeaveReplication.OwnedObjects.Add({ StreamId, StaticMeshComponent });
+			FConcertSyncReplicationActivity StreamActivity;
+			StreamActivity.EventData.SetPayload(LeaveReplication);
+			if (!SessionDatabase_Server.AddReplicationActivity(StreamActivity, ActivityId, EventId))
+			{
+				AddError(FString::Printf(TEXT("Failed to add replication stream activity to server database: %s"), *SessionDatabase_Server.GetLastError()));
+			}
+		}
+
+		// Close and re-open the client database to verify that it isn't corrupted
+		if (SessionDatabase_Client.IsValid())
+		{
+			if (!SessionDatabase_Client.Close())
+			{
+				AddError(TEXT("Failed to close client database"));
+			}
+			if (!SessionDatabase_Client.Open(TestSessionPath_Client))
+			{
+				AddError(TEXT("Failed to open client database"));
+			}
+		}
+	}
+
+	// Populate the client database from the server data
+	if (SessionDatabase_Client.IsValid())
+	{
+		// Endpoint IDs must be synced before activity, as the database enforces that the activity was from a known endpoint
+		SessionDatabase_Server.EnumerateEndpoints([this, &SessionDatabase_Client](FConcertSyncEndpointIdAndData&& InEndpointData)
+		{
+			if (!SessionDatabase_Client.SetEndpoint(InEndpointData.EndpointId, InEndpointData.EndpointData))
+			{
+				AddError(FString::Printf(TEXT("Failed to set endpoint '%s' on client database: %s"), *InEndpointData.EndpointId.ToString(), *SessionDatabase_Client.GetLastError()));
+				return false;
+			}
+			return true;
+		});
+
+		// This loop represents both what a sync command queue would do when queuing events on the server, and what the client would do when receiving them
+		// This is a very simple example that just syncs all data for every activity, but a more complex case would be skipping the event data for non-live transactions and non-head package revisions
+		SessionDatabase_Server.EnumerateActivityIdsAndEventTypes([this, &SessionDatabase_Server, &SessionDatabase_Client](const int64 InActivityId, const EConcertSyncActivityEventType InEventType)
+		{
+			switch (InEventType)
+			{
+			case EConcertSyncActivityEventType::Connection:
+			{
+				FConcertSyncConnectionActivity ConnectionActivity;
+				if (!SessionDatabase_Server.GetConnectionActivity(InActivityId, ConnectionActivity))
+				{
+					AddError(FString::Printf(TEXT("Failed to get connection activity '%s' from server database: %s"), *LexToString(InActivityId), *SessionDatabase_Server.GetLastError()));
+					return false;
+				}
+				if (!SessionDatabase_Client.SetConnectionActivity(ConnectionActivity))
+				{
+					AddError(FString::Printf(TEXT("Failed to set connection activity '%s' on client database: %s"), *LexToString(InActivityId), *SessionDatabase_Client.GetLastError()));
+					return false;
+				}
+			}
+			break;
+
+			case EConcertSyncActivityEventType::Lock:
+			{
+				FConcertSyncLockActivity LockActivity;
+				if (!SessionDatabase_Server.GetLockActivity(InActivityId, LockActivity))
+				{
+					AddError(FString::Printf(TEXT("Failed to get lock activity '%s' from server database: %s"), *LexToString(InActivityId), *SessionDatabase_Server.GetLastError()));
+					return false;
+				}
+				if (!SessionDatabase_Client.SetLockActivity(LockActivity))
+				{
+					AddError(FString::Printf(TEXT("Failed to set lock activity '%s' on client database: %s"), *LexToString(InActivityId), *SessionDatabase_Client.GetLastError()));
+					return false;
+				}
+			}
+			break;
+
+			case EConcertSyncActivityEventType::Transaction:
+			{
+				FConcertSyncTransactionActivity TransactionActivity;
+				if (!SessionDatabase_Server.GetTransactionActivity(InActivityId, TransactionActivity))
+				{
+					AddError(FString::Printf(TEXT("Failed to get transaction activity '%s' from server database: %s"), *LexToString(InActivityId), *SessionDatabase_Server.GetLastError()));
+					return false;
+				}
+				if (!SessionDatabase_Client.SetTransactionActivity(TransactionActivity))
+				{
+					AddError(FString::Printf(TEXT("Failed to set transaction activity '%s' on client database: %s"), *LexToString(InActivityId), *SessionDatabase_Client.GetLastError()));
+					return false;
+				}
+			}
+			break;
+
+			case EConcertSyncActivityEventType::Package:
+			{
+				bool bSetPackageActivitySucceeded = true;
+				bool bGetPackageActivitySucceeded = SessionDatabase_Server.GetPackageActivity(InActivityId, [&](FConcertSyncActivity&& ActivityBasePart, FConcertSyncPackageEventData& ActivityEventPart)
+				{
+					bSetPackageActivitySucceeded = SessionDatabase_Client.SetPackageActivity(ActivityBasePart, ActivityEventPart);
+				});
+
+				if (!bGetPackageActivitySucceeded)
+				{
+					AddError(FString::Printf(TEXT("Failed to get package activity '%s' from server database: %s"), *LexToString(InActivityId), *SessionDatabase_Server.GetLastError()));
+					return false;
+				}
+				if (!bSetPackageActivitySucceeded)
+				{
+					AddError(FString::Printf(TEXT("Failed to set package activity '%s' on client database: %s"), *LexToString(InActivityId), *SessionDatabase_Client.GetLastError()));
+					return false;
+				}
+			}
+			break;
+
+			case EConcertSyncActivityEventType::Replication:
+			{
+				FConcertSyncReplicationActivity ReplicationActivity;
+				if (!SessionDatabase_Server.GetReplicationActivity(InActivityId, ReplicationActivity))
+				{
+					AddError(FString::Printf(TEXT("Failed to get replication activity '%s' from server database: %s"), *LexToString(InActivityId), *SessionDatabase_Server.GetLastError()));
+					return false;
+				}
+				if (!SessionDatabase_Client.SetReplicationActivity(ReplicationActivity))
+				{
+					AddError(FString::Printf(TEXT("Failed to set replication activity '%s' on client database: %s"), *LexToString(InActivityId), *SessionDatabase_Client.GetLastError()));
+					return false;
+				}
+			}
+			break;
+
+			default:
+				AddError(TEXT("Unhandled EConcertSyncActivityEventType when populating client database"));
+				break;
+			}
+
+			return true;
+		});
+	}
+
+	// Verify the two databases have the same content
+	if (SessionDatabase_Server.IsValid() && SessionDatabase_Client.IsValid())
+	{
+		SessionDatabase_Server.EnumerateEndpoints([this, &SessionDatabase_Client](FConcertSyncEndpointIdAndData&& InEndpointData)
+		{
+			FConcertSyncEndpointData ClientEndpointData;
+			if (!SessionDatabase_Client.GetEndpoint(InEndpointData.EndpointId, ClientEndpointData))
+			{
+				AddError(FString::Printf(TEXT("Failed to get endpoint '%s' from client database: %s"), *InEndpointData.EndpointId.ToString(), *SessionDatabase_Client.GetLastError()));
+				return false;
+			}
+			return true;
+		});
+
+		SessionDatabase_Server.EnumerateActivityIdsAndEventTypes([this, &SessionDatabase_Client](const int64 InActivityId, const EConcertSyncActivityEventType InEventType)
+		{
+			FConcertSyncActivity ClientActivity;
+			if (!SessionDatabase_Client.GetActivity(InActivityId, ClientActivity))
+			{
+				AddError(FString::Printf(TEXT("Failed to get activity '%s' from client database: %s"), *LexToString(InActivityId), *SessionDatabase_Client.GetLastError()));
+				return false;
+			}
+			if (InEventType != ClientActivity.EventType)
+			{
+				AddError(FString::Printf(TEXT("Activity '%s' from client database was the incorrect type"), *LexToString(InActivityId)));
+				return false;
+			}
+			return true;
+		});
+	}
+
+	if (SessionDatabase_Server.IsValid() && !SessionDatabase_Server.Close())
+	{
+		AddError(TEXT("Failed to close server database"));
+	}
+
+	if (SessionDatabase_Client.IsValid() && !SessionDatabase_Client.Close())
+	{
+		AddError(TEXT("Failed to close client database"));
+	}
+
+	IFileManager::Get().DeleteDirectory(*TestSessionPath_Server, false, true);
+	IFileManager::Get().DeleteDirectory(*TestSessionPath_Client, false, true);
+
+	return true;
+}
diff --git a/Engine/Plugins/Developer/Concert/ConcertSync/ConcertSyncTest/Source/ConcertSyncTest/Private/Database/ReplicationActivity.spec.cpp b/Engine/Plugins/Developer/Concert/ConcertSync/ConcertSyncTest/Source/ConcertSyncTest/Private/Database/ReplicationActivity.spec.cpp
new file mode 100644
--- /dev/null
+++ b/Engine/Plugins/Developer/Concert/ConcertSync/ConcertSyncTest/Source/ConcertSyncTest/Private/Database/ReplicationActivity.spec.cpp
@@ -0,0 +1,180 @@
+﻿// Copyright Epic Games, Inc. All Rights Reserved.
+
+#include "ConcertSyncSessionDatabase.h"
+#include "Replication/Data/ConcertPropertySelection.h"
+#include "Replication/Data/ObjectReplicationMap.h"
+#include "Replication/Data/ReplicationStream.h"
+#include "Replication/Messages/ReplicationActivity.h"
+
+#include "Components/StaticMeshComponent.h"
+#include "HAL/FileManager.h"
+#include "Misc/AutomationTest.h"
+#include "Misc/Paths.h"
+#include "Templates/UniquePtr.h"
+
+namespace UE::ConcertSyncTests
+{
+	BEGIN_DEFINE_SPEC(FConcertDatabaseTest_ReplicationActivity, "Editor.Concert.Database.ReplicationActivity", EAutomationTestFlags::EditorContext | EAutomationTestFlags::EngineFilter)
+		TUniquePtr<FConcertSyncSessionDatabase> Database;
+	
+		static FString GetDatabasePath() { return FPaths::ProjectIntermediateDir() / TEXT("ConcertDatabaseTest_Server"); }
+		static FConcertSyncReplicationPayload_LeaveReplication MakeLeaveReplicationActivity(const FSoftObjectPath& ObjectPath, const FGuid& StreamId)
+		{
+			FConcertReplicationStream Stream;
+			Stream.BaseDescription.Identifier = StreamId;
+			FConcertReplicatedObjectInfo& ObjectInfo = Stream.BaseDescription.ReplicationMap.ReplicatedObjects.Add(ObjectPath);
+			ObjectInfo.ClassPath = UStaticMeshComponent::StaticClass();
+			const TOptional<FConcertPropertyChain> Property_RelativeX = FConcertPropertyChain::CreateFromPath(*UStaticMeshComponent::StaticClass(), { TEXT("RelativeLocation"), TEXT("X") });
+			if (ensure(Property_RelativeX))
+			{
+				ObjectInfo.PropertySelection.ReplicatedProperties.Add(*Property_RelativeX);
+			}
+			
+			FConcertSyncReplicationPayload_LeaveReplication Payload;
+			Payload.Streams.Add(Stream);
+			Payload.OwnedObjects.Add(FConcertObjectInStreamID{ StreamId, ObjectPath });
+				
+			return Payload;
+		}
+
+		static FConcertSyncReplicationActivity MakeActivity(FConcertSyncReplicationPayload_LeaveReplication Payload, const FGuid& ClientId)
+		{
+			FConcertSyncReplicationActivity Activity(Payload);
+			Activity.EndpointId = ClientId;
+			return Activity;
+		}
+	END_DEFINE_SPEC(FConcertDatabaseTest_ReplicationActivity);
+
+	/** This tests that clients respect sync control (see FConcertReplication_ChangeSyncControl). */
+	void FConcertDatabaseTest_ReplicationActivity::Define()
+	{
+		BeforeEach([this]()
+		{
+			// Delete the files in case we did not clean up in a previous run (e.g. we hit a break point and then terminated the app).
+			IFileManager::Get().DeleteDirectory(*GetDatabasePath(), false, true);
+			
+			Database = MakeUnique<FConcertSyncSessionDatabase>();
+			Database->Open(GetDatabasePath());
+		});
+
+		AfterEach([this]()
+		{
+			Database->Close();
+			IFileManager::Get().DeleteDirectory(*GetDatabasePath(), false, true);
+			Database.Reset();
+		});
+		
+		struct FActivityEntry
+		{
+			FConcertSyncReplicationActivity Activity;
+			int64 ActivityId = INDEX_NONE;
+			int64 ReplicationEventId = INDEX_NONE;
+				
+			FActivityEntry(FConcertSyncReplicationActivity Activity) : Activity(MoveTemp(Activity)) {}
+
+			bool operator==(const FActivityEntry& Right) const
+			{
+				return ActivityId == Right.ActivityId && ReplicationEventId == Right.ReplicationEventId && Activity.EventData == Right.Activity.EventData;
+			}
+		};
+
+		It("Add and retrieve FConcertSyncReplicationPayload_LeaveReplication", [this]
+		{
+			const FGuid StreamId = FGuid::NewGuid();
+			const FSoftObjectPath StaticMeshComponent(TEXT("/Game/Maps.Map:PersistentLevel.Cube.StaticMeshComponent0"));
+			
+			const FConcertSyncReplicationPayload_LeaveReplication OriginalPayload = MakeLeaveReplicationActivity(StaticMeshComponent, StreamId);
+			const FConcertSyncReplicationActivity OriginalActivity(OriginalPayload);
+			int64 ActivityId, ReplicationEventId;
+			const bool bAddedActivity = Database->AddReplicationActivity(OriginalActivity, ActivityId, ReplicationEventId);
+			if (!bAddedActivity)
+			{
+				AddError(TEXT("Failed to store activity in database"));
+				return;
+			}
+
+			FConcertSyncReplicationActivity RetrievedActivity;
+			const bool bRetrievedActivity = Database->GetReplicationActivity(ActivityId, RetrievedActivity);
+			if (!bRetrievedActivity)
+			{
+				AddError(TEXT("Failed to retrieve activity from database"));
+				return;
+			}
+
+			FConcertSyncReplicationPayload_LeaveReplication RetriedPayload;
+			const bool bReadStreamChange = RetrievedActivity.EventData.GetPayload(RetriedPayload);
+			if (!bReadStreamChange)
+			{
+				AddError(TEXT("Failed to read payload as FConcertSyncReplicationPayload_StreamChange"));
+				return;
+			}
+
+			TestEqual(TEXT("Stream changes equal"), OriginalPayload, RetriedPayload);
+		});
+
+		It("EnumerateReplicationActivities", [this]
+		{
+			const FGuid StreamId = FGuid::NewGuid();
+			TArray<FActivityEntry> Activities_OriginalOrder =
+			{
+				FConcertSyncReplicationActivity(MakeLeaveReplicationActivity(FSoftObjectPath(TEXT("/Game/Maps.Map:PersistentLevel.Cube.StaticMeshComponent0")), StreamId)),
+				FConcertSyncReplicationActivity(MakeLeaveReplicationActivity(FSoftObjectPath(TEXT("/Game/Maps.Map:PersistentLevel.Cube.StaticMeshComponent1")), StreamId)),
+				FConcertSyncReplicationActivity(MakeLeaveReplicationActivity(FSoftObjectPath(TEXT("/Game/Maps.Map:PersistentLevel.Cube.StaticMeshComponent2")), StreamId))
+			};
+			for (FActivityEntry& Entry : Activities_OriginalOrder)
+			{
+				Database->AddReplicationActivity(Entry.Activity, Entry.ActivityId, Entry.ReplicationEventId);
+			}
+			
+			TArray<FActivityEntry> Activities_EnumerationOrder;
+			Database->EnumerateReplicationActivities([&Activities_EnumerationOrder](FConcertSyncReplicationActivity&& Activity)
+			{
+				Activities_EnumerationOrder.Emplace(MoveTemp(Activity));
+				return true;
+			});
+
+			if (Activities_EnumerationOrder.Num() && Activities_OriginalOrder.Num())
+			{
+				AddError(TEXT("Counts do not match"));
+				return;
+			}
+			for (int32 i = 0; i < Activities_EnumerationOrder.Num(); ++i)
+			{
+				TestEqual(TEXT("Activities match"), Activities_EnumerationOrder[i], Activities_OriginalOrder[i]);
+			}
+		});
+		
+		It("GetReplicationMaxEventIdByClientAndType", [this]
+		{
+			const FGuid StreamId = FGuid::NewGuid();
+			const FGuid ClientOne = FGuid::NewGuid();
+			const FGuid ClientTwo = FGuid::NewGuid();
+			TArray<FActivityEntry> Activities_OriginalOrder =
+			{
+				MakeActivity(MakeLeaveReplicationActivity(FSoftObjectPath(TEXT("/Game/Maps.Map:PersistentLevel.Cube.StaticMeshComponent0")), StreamId), ClientOne),
+				MakeActivity(MakeLeaveReplicationActivity(FSoftObjectPath(TEXT("/Game/Maps.Map:PersistentLevel.Cube.StaticMeshComponent1")), StreamId), ClientOne),
+				MakeActivity(MakeLeaveReplicationActivity(FSoftObjectPath(TEXT("/Game/Maps.Map:PersistentLevel.Cube.StaticMeshComponent2")), StreamId), ClientTwo)
+			};
+			for (FActivityEntry& Entry : Activities_OriginalOrder)
+			{
+				Database->AddReplicationActivity(Entry.Activity, Entry.ActivityId, Entry.ReplicationEventId);
+			}
+			
+			int64 EventId_ClientOne = INDEX_NONE;
+			int64 EventId_ClientTwo = INDEX_NONE;
+			int64 EventId_Invalid = INDEX_NONE;
+			const bool bSuccess_Query1 = Database->GetReplicationMaxEventIdByClientAndType(ClientOne, EConcertSyncReplicationActivityType::LeaveReplication, EventId_ClientOne);
+			const bool bSuccess_Query2 = Database->GetReplicationMaxEventIdByClientAndType(ClientTwo, EConcertSyncReplicationActivityType::LeaveReplication, EventId_ClientTwo);
+			const bool bSuccess_InvalidQuery = Database->GetReplicationMaxEventIdByClientAndType(FGuid::NewGuid(), EConcertSyncReplicationActivityType::LeaveReplication, EventId_Invalid);
+			
+			TestTrue(TEXT("Query 1"), bSuccess_Query1);
+			TestEqual(TEXT("Client 1 Max EventId"), EventId_ClientOne, Activities_OriginalOrder[1].ReplicationEventId);
+			TestTrue(TEXT("Query 2"), bSuccess_Query2);
+			TestEqual(TEXT("Client 2 Max EventId"), EventId_ClientTwo, Activities_OriginalOrder[2].ReplicationEventId);
+			TestTrue(TEXT("Invalid query"), bSuccess_InvalidQuery);
+			// Sql query returns NULL here (since nothing is found).
+			// Under the hood, sqlite3_value_int64 is used to get the value of the returned column... NULL is converted to 0 by sqlite3_value_int64.
+			TestEqual(TEXT("Unfound returns NULL"), EventId_Invalid, static_cast<int64>(0));
+		});
+	}
+}
