Fix duplicate loca keys

diff --git a/Engine/Plugins/Developer/Concert/ConcertSync/ConcertSyncCore/Source/ConcertSyncCore/Public/ConcertSyncSessionTypes.cpp b/Engine/Plugins/Developer/Concert/ConcertSync/ConcertSyncCore/Source/ConcertSyncCore/Public/ConcertSyncSessionTypes.cpp
--- a/Engine/Plugins/Developer/Concert/ConcertSync/ConcertSyncCore/Source/ConcertSyncCore/Public/ConcertSyncSessionTypes.cpp
+++ b/Engine/Plugins/Developer/Concert/ConcertSync/ConcertSyncCore/Source/ConcertSyncCore/Public/ConcertSyncSessionTypes.cpp
@@ -841,12 +841,12 @@ FText FConcertSyncReplicationActivitySummary::CreateDisplayText(const bool InUse
 			const bool bRead = GetSummaryData(Content);
 			if (!bRead)
 			{
-				return LOCTEXT("CreateDisplayTextForUser_LeftActivity", "No data.");
+				return LOCTEXT("CreateDisplayText.LeftActivity.NoData", "No data.");
 			}
 			
 			FFormatNamedArguments Arguments;
 			Arguments.Add(TEXT("NumObjects"), ActivitySummaryUtil::ToRichTextBold(FText::AsNumber(Content.OwnedObjects.Num()), InUseRichText));
-			return FText::Format(LOCTEXT("CreateDisplayTextForUser.LeftActivityFmt", "Stopped replicating {NumObjects} objects."), Arguments);
+			return FText::Format(LOCTEXT("CreateDisplayText.LeftActivity.Fmt", "Stopped replicating {NumObjects} objects."), Arguments);
 			
 		}
 	case EConcertSyncReplicationActivityType::None: [[fallthrough]];
@@ -865,13 +865,13 @@ FText FConcertSyncReplicationActivitySummary::CreateDisplayTextForUser(const FTe
 			const bool bRead = GetSummaryData(Content);
 			if (!bRead)
 			{
-				return LOCTEXT("CreateDisplayTextForUser_LeftActivity", "No data.");
+				return LOCTEXT("CreateDisplayTextForUser.LeftActivity.NoData", "No data.");
 			}
 			
 			FFormatNamedArguments Arguments;
 			Arguments.Add(TEXT("UserName"), ActivitySummaryUtil::ToRichTextBold(InUserDisplayName, InUseRichText));
 			Arguments.Add(TEXT("NumObjects"), ActivitySummaryUtil::ToRichTextBold(FText::AsNumber(Content.OwnedObjects.Num()), InUseRichText));
-			return FText::Format(LOCTEXT("CreateDisplayTextForUser.LeftActivityFmt", "{UserName} stopped replicating {NumObjects} objects."), Arguments);
+			return FText::Format(LOCTEXT("CreateDisplayTextForUser.LeftActivity.Fmt", "{UserName} stopped replicating {NumObjects} objects."), Arguments);
 		}
 	case EConcertSyncReplicationActivityType::None: [[fallthrough]];
 	default: checkNoEntry(); return FConcertSyncActivitySummary::CreateDisplayText(InUseRichText);
