Fix copy-paste mistake in check

diff --git a/Engine/Plugins/Developer/Concert/ConcertSync/ConcertSyncCore/Source/ConcertSyncCore/Private/Replication/ChangeStreamSharedUtils.cpp b/Engine/Plugins/Developer/Concert/ConcertSync/ConcertSyncCore/Source/ConcertSyncCore/Private/Replication/ChangeStreamSharedUtils.cpp
--- a/Engine/Plugins/Developer/Concert/ConcertSync/ConcertSyncCore/Source/ConcertSyncCore/Private/Replication/ChangeStreamSharedUtils.cpp
+++ b/Engine/Plugins/Developer/Concert/ConcertSync/ConcertSyncCore/Source/ConcertSyncCore/Private/Replication/ChangeStreamSharedUtils.cpp
@@ -128,7 +128,7 @@ namespace UE::ConcertSyncCore::Replication::ChangeStreamUtils
 					ObjectInfo.PropertySelection = SelectionToSet;
 				}
 				// ... but one of these cases must occur
-				checkf(ObjectInfo.ClassPath.IsValid(), TEXT("Request did not validate Properties!"));
+				checkf(!ObjectInfo.PropertySelection.ReplicatedProperties.IsEmpty(), TEXT("Request did not validate Properties!"));
 			}
 		}
 
