Save replication mute & unmute requests into the Concert ledger

diff --git a/Engine/Plugins/Developer/Concert/ConcertSync/ConcertSyncCore/Source/ConcertSyncCore/Private/Replication/Messages/ReplicationActivity.cpp b/Engine/Plugins/Developer/Concert/ConcertSync/ConcertSyncCore/Source/ConcertSyncCore/Private/Replication/Messages/ReplicationActivity.cpp
--- a/Engine/Plugins/Developer/Concert/ConcertSync/ConcertSyncCore/Source/ConcertSyncCore/Private/Replication/Messages/ReplicationActivity.cpp
+++ b/Engine/Plugins/Developer/Concert/ConcertSync/ConcertSyncCore/Source/ConcertSyncCore/Private/Replication/Messages/ReplicationActivity.cpp
@@ -4,6 +4,8 @@
 
 #include "Internationalization/Text.h"
 
+#define LOCTEXT_NAMESPACE "ReplicationActivity"
+
 bool operator==(const FConcertSyncReplicationEvent& Left, const FConcertSyncReplicationEvent& Right)
 {
 	if (Left.ActivityType != Right.ActivityType)
@@ -11,7 +13,7 @@ bool operator==(const FConcertSyncReplicationEvent& Left, const FConcertSyncRepl
 		return false;
 	}
 
-	static_assert(static_cast<int32>(EConcertSyncReplicationActivityType::Count) == 2, "If you added an EConcertSyncReplicationActivityType entry, update this switch");
+	static_assert(static_cast<int32>(EConcertSyncReplicationActivityType::Count) == 3, "If you added an EConcertSyncReplicationActivityType entry, update this switch");
 	switch (Left.ActivityType)
 	{
 	case EConcertSyncReplicationActivityType::None: return true;
@@ -23,6 +25,14 @@ bool operator==(const FConcertSyncReplicationEvent& Left, const FConcertSyncRepl
 			const bool bRightSucceeded = Right.GetPayload(RightContent);
 			return bLeftSucceeded && bRightSucceeded && LeftContent == RightContent;
 		}
+	case EConcertSyncReplicationActivityType::Mute:
+		{
+			FConcertSyncReplicationPayload_Mute LeftContent;
+			FConcertSyncReplicationPayload_Mute RightContent;
+			const bool bLeftSucceeded = Left.GetPayload(LeftContent);
+			const bool bRightSucceeded = Right.GetPayload(RightContent);
+			return bLeftSucceeded && bRightSucceeded && LeftContent == RightContent;
+		}
 	default: checkNoEntry(); return false;
 	}
 }
@@ -32,7 +42,7 @@ FConcertSyncReplicationActivitySummary FConcertSyncReplicationActivitySummary::C
 	FConcertSyncReplicationActivitySummary Summary;
 	Summary.ActivityType = InEvent.ActivityType;
 
-	static_assert(static_cast<uint8>(EConcertSyncReplicationActivityType::Count) == 2, "If you changed the enum entries, update this switch");
+	static_assert(static_cast<uint8>(EConcertSyncReplicationActivityType::Count) == 3, "If you changed the enum entries, update this switch");
 	switch (Summary.ActivityType)
 	{
 	case EConcertSyncReplicationActivityType::LeaveReplication:
@@ -44,9 +54,59 @@ FConcertSyncReplicationActivitySummary FConcertSyncReplicationActivitySummary::C
 			Summary.Payload.SetTypedPayload(Content);
 		}
 		break;
+	case EConcertSyncReplicationActivityType::Mute:
+		{
+			FConcertSyncReplicationPayload_Mute MuteData;
+			InEvent.GetPayload(MuteData);
+			
+			const FConcertSyncReplicationSummary_Mute Content { MuteData.Request };
+			Summary.Payload.SetTypedPayload(Content);
+		}
+		break;
 	case EConcertSyncReplicationActivityType::None: [[fallthrough]];
 	default: checkNoEntry(); break;
 	}
 	
 	return Summary;
-}
\ No newline at end of file
+}
+
+FText FConcertSyncReplicationActivitySummary::ToDisplayTitle() const
+{
+	static_assert(static_cast<uint8>(EConcertSyncReplicationActivityType::Count) == 3, "If you changed the enum entries, update this switch");
+	switch (ActivityType)
+	{
+	case EConcertSyncReplicationActivityType::LeaveReplication: return LOCTEXT("Title.LeftReplication", "Left Replication");
+	case EConcertSyncReplicationActivityType::Mute:
+		{
+			FConcertSyncReplicationSummary_Mute SummaryData;
+			if (!GetSummaryData(SummaryData))
+			{
+				return LOCTEXT("Title.Mute.FailedToGetData", "Pause / Resume");
+			}
+
+			const bool bHasMuted = !SummaryData.Request.ObjectsToMute.IsEmpty();
+			const bool bHasUnmuted = !SummaryData.Request.ObjectsToUnmute.IsEmpty();
+			if (bHasMuted && bHasUnmuted)
+			{
+				return LOCTEXT("Title.Mute.PauseAndResume", "Pause & Resume");
+			}
+
+			if (bHasMuted)
+			{
+				return LOCTEXT("Title.Mute.Pause", "Pause replication");
+			}
+			
+			if (bHasUnmuted)
+			{
+				return LOCTEXT("Title.Mute.Resume", "Resume replication");
+			}
+
+			return LOCTEXT("Title.Mute.Resume", "Pause / Resume (empty)");
+		}
+	
+	case EConcertSyncReplicationActivityType::None: [[fallthrough]];
+	default: checkNoEntry(); return FText::GetEmpty();
+	}
+}
+
+#undef LOCTEXT_NAMESPACE
\ No newline at end of file
diff --git a/Engine/Plugins/Developer/Concert/ConcertSync/ConcertSyncCore/Source/ConcertSyncCore/Public/ConcertSyncSessionTypes.cpp b/Engine/Plugins/Developer/Concert/ConcertSync/ConcertSyncCore/Source/ConcertSyncCore/Public/ConcertSyncSessionTypes.cpp
--- a/Engine/Plugins/Developer/Concert/ConcertSync/ConcertSyncCore/Source/ConcertSyncCore/Public/ConcertSyncSessionTypes.cpp
+++ b/Engine/Plugins/Developer/Concert/ConcertSync/ConcertSyncCore/Source/ConcertSyncCore/Public/ConcertSyncSessionTypes.cpp
@@ -830,25 +830,114 @@ FText FConcertSyncPackageActivitySummary::CreateDisplayTextForUser(const FText I
 	return FText::Format(FormatPattern, Arguments);
 }
 
-FText FConcertSyncReplicationActivitySummary::CreateDisplayText(const bool InUseRichText) const
+namespace UE::ConcertSyncCore::ActivitySummary::Private
 {
-	static_assert(static_cast<uint8>(EConcertSyncReplicationActivityType::Count) == 2, "If you added an EConcertSyncReplicationActivityType entry, update this switch");
-	switch (ActivityType)
+	static FText CreateDisplayText_LeaveReplication(const FConcertSyncReplicationActivitySummary& Summary, const bool bInUseRichText)
 	{
-	case EConcertSyncReplicationActivityType::LeaveReplication:
+		FConcertSyncReplicationSummary_LeaveReplication Content;
+		const bool bRead = Summary.GetSummaryData(Content);
+		if (!bRead)
 		{
-			FConcertSyncReplicationSummary_LeaveReplication Content;
-			const bool bRead = GetSummaryData(Content);
-			if (!bRead)
-			{
-				return LOCTEXT("CreateDisplayText.LeftActivity.NoData", "No data.");
-			}
+			return LOCTEXT("CreateDisplayText.LeftActivity.NoData", "No data.");
+		}
 			
-			FFormatNamedArguments Arguments;
-			Arguments.Add(TEXT("NumObjects"), ActivitySummaryUtil::ToRichTextBold(FText::AsNumber(Content.OwnedObjects.Num()), InUseRichText));
-			return FText::Format(LOCTEXT("CreateDisplayText.LeftActivity.Fmt", "Stopped replicating {NumObjects} objects."), Arguments);
+		FFormatNamedArguments Arguments;
+		Arguments.Add(TEXT("NumObjects"), ActivitySummaryUtil::ToRichTextBold(FText::AsNumber(Content.OwnedObjects.Num()), bInUseRichText));
+		return FText::Format(LOCTEXT("CreateDisplayText.LeftActivity.Fmt", "Stopped replicating {NumObjects} {NumObjects}|plural(one=object, other=objects)."), Arguments);
+	}
+	static FText CreateDisplayTextForUser_LeaveReplication(const FConcertSyncReplicationActivitySummary& Summary, const FText InUserDisplayName, const bool InUseRichText)
+	{
+		FConcertSyncReplicationSummary_LeaveReplication Content;
+		const bool bRead = Summary.GetSummaryData(Content);
+		if (!bRead)
+		{
+			return LOCTEXT("CreateDisplayTextForUser.LeftActivity.NoData", "No data.");
+		}
+			
+		FFormatNamedArguments Arguments;
+		Arguments.Add(TEXT("UserName"), ActivitySummaryUtil::ToRichTextBold(InUserDisplayName, InUseRichText));
+		Arguments.Add(TEXT("NumObjects"), ActivitySummaryUtil::ToRichTextBold(FText::AsNumber(Content.OwnedObjects.Num()), InUseRichText));
+		return FText::Format(LOCTEXT("CreateDisplayTextForUser.LeftActivity.Fmt", "{UserName} stopped replicating {NumObjects} {NumObjects}|plural(one=object, other=objects)."), Arguments);
+	}
+
+	static FText SelectText(bool bMuted, bool bUnmuted, const FText& MutedOnly, const FText& UnmutedOnly, const FText& MutedAndUnmuted)
+	{
+		if (bMuted && bUnmuted)
+		{
+			return MutedAndUnmuted;
+		}
+
+		if (bMuted)
+		{
+			return MutedOnly;
+		}
+
+		if (bUnmuted)
+		{
+			return UnmutedOnly;
+		}
+
+		return LOCTEXT("CreateDisplayText.Mute.None", "No changes");
+	}
+	
+	static FText CreateDisplayText_Mute(const FConcertSyncReplicationActivitySummary& Summary, const bool bInUseRichText)
+	{
+		FConcertSyncReplicationSummary_Mute Content;
+		const bool bRead = Summary.GetSummaryData(Content);
+		if (!bRead)
+		{
+			return LOCTEXT("CreateDisplayText.LeftActivity.NoData", "No data.");
+		}
+		const FConcertReplication_ChangeMuteState_Request& Request = Content.Request;
 			
+		FFormatNamedArguments Arguments;
+		Arguments.Add(TEXT("NumMuted"), ActivitySummaryUtil::ToRichTextBold(FText::AsNumber(Request.ObjectsToMute.Num()), bInUseRichText));
+		Arguments.Add(TEXT("NumUnmuted"), ActivitySummaryUtil::ToRichTextBold(FText::AsNumber(Request.ObjectsToUnmute.Num()), bInUseRichText));
+		
+		const FText TextToFormat = SelectText(
+			!Request.ObjectsToMute.IsEmpty(),
+			!Request.ObjectsToUnmute.IsEmpty(),
+			LOCTEXT("CreateDisplayText.MutedOnly.Fmt", "Paused {NumMuted} {NumMuted}|plural(one=object, other=objects)."),
+			LOCTEXT("CreateDisplayText.UnmutedOnly.Fmt", "Resumed {NumUnmuted} {NumUnmuted}|plural(one=object, other=objects)."),
+			LOCTEXT("CreateDisplayText.MutedAndUnmuted.Fmt", "Paused {NumObjects} and resumed {NumUnmuted} objects.")
+			);
+		return FText::Format(TextToFormat, Arguments);
+	}
+	static FText CreateDisplayTextForUser_Mute(const FConcertSyncReplicationActivitySummary& Summary, const FText InUserDisplayName, const bool bInUseRichText)
+	{
+		FConcertSyncReplicationSummary_Mute Content;
+		const bool bRead = Summary.GetSummaryData(Content);
+		if (!bRead)
+		{
+			return LOCTEXT("CreateDisplayText.LeftActivity.NoData", "No data.");
 		}
+		const FConcertReplication_ChangeMuteState_Request& Request = Content.Request;
+			
+		FFormatNamedArguments Arguments;
+		Arguments.Add(TEXT("NumMuted"), ActivitySummaryUtil::ToRichTextBold(FText::AsNumber(Request.ObjectsToMute.Num()), bInUseRichText));
+		Arguments.Add(TEXT("NumUnmuted"), ActivitySummaryUtil::ToRichTextBold(FText::AsNumber(Request.ObjectsToUnmute.Num()), bInUseRichText));
+		Arguments.Add(TEXT("UserName"), ActivitySummaryUtil::ToRichTextBold(InUserDisplayName, bInUseRichText));
+		
+		const FText TextToFormat = SelectText(
+			!Request.ObjectsToMute.IsEmpty(),
+			!Request.ObjectsToUnmute.IsEmpty(),
+			LOCTEXT("CreateDisplayTextForUser.MutedOnly.Fmt", "{UserName} paused {NumMuted} {NumMuted}|plural(one=object, other=objects)."),
+			LOCTEXT("CreateDisplayTextForUser.UnmutedOnly.Fmt", "{UserName} resumed {NumUnmuted} {NumUnmuted}|plural(one=object, other=objects)."),
+			LOCTEXT("CreateDisplayTextForUser.MutedAndUnmuted.Fmt", "{UserName} paused {NumObjects} and resumed {NumUnmuted} objects.")
+			);
+		return FText::Format(TextToFormat, Arguments);
+	}
+}
+
+FText FConcertSyncReplicationActivitySummary::CreateDisplayText(const bool InUseRichText) const
+{
+	using namespace UE::ConcertSyncCore::ActivitySummary::Private;
+	
+	static_assert(static_cast<uint8>(EConcertSyncReplicationActivityType::Count) == 3, "If you added an EConcertSyncReplicationActivityType entry, update this switch");
+	switch (ActivityType)
+	{
+	case EConcertSyncReplicationActivityType::LeaveReplication: return CreateDisplayText_LeaveReplication(*this, InUseRichText);
+	case EConcertSyncReplicationActivityType::Mute: return CreateDisplayText_Mute(*this, InUseRichText);
 	case EConcertSyncReplicationActivityType::None: [[fallthrough]];
 		default: checkNoEntry(); return FConcertSyncActivitySummary::CreateDisplayText(InUseRichText);
 	}
@@ -856,23 +945,13 @@ FText FConcertSyncReplicationActivitySummary::CreateDisplayText(const bool InUse
 
 FText FConcertSyncReplicationActivitySummary::CreateDisplayTextForUser(const FText InUserDisplayName, const bool InUseRichText) const
 {
-	static_assert(static_cast<uint8>(EConcertSyncReplicationActivityType::Count) == 2, "If you added an EConcertSyncReplicationActivityType entry, update this switch");
+	using namespace UE::ConcertSyncCore::ActivitySummary::Private;
+	
+	static_assert(static_cast<uint8>(EConcertSyncReplicationActivityType::Count) == 3, "If you added an EConcertSyncReplicationActivityType entry, update this switch");
 	switch (ActivityType)
 	{
-	case EConcertSyncReplicationActivityType::LeaveReplication:
-		{
-			FConcertSyncReplicationSummary_LeaveReplication Content;
-			const bool bRead = GetSummaryData(Content);
-			if (!bRead)
-			{
-				return LOCTEXT("CreateDisplayTextForUser.LeftActivity.NoData", "No data.");
-			}
-			
-			FFormatNamedArguments Arguments;
-			Arguments.Add(TEXT("UserName"), ActivitySummaryUtil::ToRichTextBold(InUserDisplayName, InUseRichText));
-			Arguments.Add(TEXT("NumObjects"), ActivitySummaryUtil::ToRichTextBold(FText::AsNumber(Content.OwnedObjects.Num()), InUseRichText));
-			return FText::Format(LOCTEXT("CreateDisplayTextForUser.LeftActivity.Fmt", "{UserName} stopped replicating {NumObjects} objects."), Arguments);
-		}
+	case EConcertSyncReplicationActivityType::LeaveReplication: return CreateDisplayTextForUser_LeaveReplication(*this, InUserDisplayName, InUseRichText);
+	case EConcertSyncReplicationActivityType::Mute: return CreateDisplayTextForUser_Mute(*this, InUserDisplayName, InUseRichText);
 	case EConcertSyncReplicationActivityType::None: [[fallthrough]];
 	default: checkNoEntry(); return FConcertSyncActivitySummary::CreateDisplayText(InUseRichText);
 	}
diff --git a/Engine/Plugins/Developer/Concert/ConcertSync/ConcertSyncCore/Source/ConcertSyncCore/Public/Replication/Messages/Muting.h b/Engine/Plugins/Developer/Concert/ConcertSync/ConcertSyncCore/Source/ConcertSyncCore/Public/Replication/Messages/Muting.h
--- a/Engine/Plugins/Developer/Concert/ConcertSync/ConcertSyncCore/Source/ConcertSyncCore/Public/Replication/Messages/Muting.h
+++ b/Engine/Plugins/Developer/Concert/ConcertSync/ConcertSyncCore/Source/ConcertSyncCore/Public/Replication/Messages/Muting.h
@@ -83,6 +83,12 @@ struct FConcertReplication_ChangeMuteState_Request
 	UPROPERTY()
 	TMap<FSoftObjectPath, FConcertReplication_ObjectMuteSetting> ObjectsToUnmute;
 
+	friend bool operator==(const FConcertReplication_ChangeMuteState_Request& Left, const FConcertReplication_ChangeMuteState_Request& Right)
+	{
+		return Left.ObjectsToMute.OrderIndependentCompareEqual(Right.ObjectsToMute) && Left.ObjectsToUnmute.OrderIndependentCompareEqual(Right.ObjectsToUnmute);
+	}
+	friend bool operator!=(const FConcertReplication_ChangeMuteState_Request&, const FConcertReplication_ChangeMuteState_Request&) = default;
+
 	/** @return Whether this requests makes no changes */
 	bool IsEmpty() const { return ObjectsToMute.IsEmpty() && ObjectsToUnmute.IsEmpty(); }
 };
diff --git a/Engine/Plugins/Developer/Concert/ConcertSync/ConcertSyncCore/Source/ConcertSyncCore/Public/Replication/Messages/ReplicationActivity.h b/Engine/Plugins/Developer/Concert/ConcertSync/ConcertSyncCore/Source/ConcertSyncCore/Public/Replication/Messages/ReplicationActivity.h
--- a/Engine/Plugins/Developer/Concert/ConcertSync/ConcertSyncCore/Source/ConcertSyncCore/Public/Replication/Messages/ReplicationActivity.h
+++ b/Engine/Plugins/Developer/Concert/ConcertSync/ConcertSyncCore/Source/ConcertSyncCore/Public/Replication/Messages/ReplicationActivity.h
@@ -3,6 +3,7 @@
 #pragma once
 
 #include "ConcertSyncSessionTypes.h"
+#include "Muting.h"
 #include "Containers/Array.h"
 #include "Containers/UnrealString.h"
 #include "Replication/Data/ObjectIds.h"
@@ -20,6 +21,9 @@ enum class EConcertSyncReplicationActivityType : uint8
 	/** Client left the replication session (either because they left the MU session or because they sent a FConcertReplication_LeaveEvent). */
 	LeaveReplication,
 
+	/** Client muted or unmuted some objects. */
+	Mute,
+
 	// ADD NEW ENTRIES ABOVE
 	Count
 };
@@ -55,13 +59,40 @@ struct FConcertSyncReplicationSummary_LeaveReplication
 	TArray<FConcertObjectInStreamID> OwnedObjects;
 };
 
+/** Stores objects that were muted / unmuted */
+USTRUCT()
+struct FConcertSyncReplicationPayload_Mute
+{
+	GENERATED_BODY()
+
+	/** The request that changed mute state */
+	UPROPERTY()
+	FConcertReplication_ChangeMuteState_Request Request;
+	
+	friend bool operator==(const FConcertSyncReplicationPayload_Mute&, const FConcertSyncReplicationPayload_Mute&) = default;
+	friend bool operator!=(const FConcertSyncReplicationPayload_Mute&, const FConcertSyncReplicationPayload_Mute&) = default;
+};
+
+/** Info displayed in the UI. */
+USTRUCT()
+struct FConcertSyncReplicationSummary_Mute
+{
+	GENERATED_BODY()
+	
+	/** The request that changed mute state */
+	UPROPERTY()
+	FConcertReplication_ChangeMuteState_Request Request;
+};
+
 namespace UE::ConcertSyncCore
 {
 	inline FString GetReplicationActivityPayloadTypePathName(EConcertSyncReplicationActivityType Type)
 	{
+		static_assert(static_cast<int32>(EConcertSyncReplicationActivityType::Count) == 3, "If you added an EConcertSyncReplicationActivityType entry, update this switch");
 		switch (Type)
 		{
 		case EConcertSyncReplicationActivityType::LeaveReplication: return FConcertSyncReplicationPayload_LeaveReplication::StaticStruct()->GetPathName();
+		case EConcertSyncReplicationActivityType::Mute: return FConcertSyncReplicationPayload_Mute::StaticStruct()->GetPathName();
 		default: ensureMsgf(false, TEXT("Unknown replication activity type (%u)"), static_cast<uint8>(Type)); return TEXT("");
 		}
 	}
@@ -99,12 +130,22 @@ struct FConcertSyncReplicationEvent
 		ActivityType = EConcertSyncReplicationActivityType::LeaveReplication;
 		Payload.SetTypedPayload(Data);
 	}
+	void SetPayload(const FConcertSyncReplicationPayload_Mute& Data)
+	{
+		ActivityType = EConcertSyncReplicationActivityType::Mute;
+		Payload.SetTypedPayload(Data);
+	}
 	
 	bool GetPayload(FConcertSyncReplicationPayload_LeaveReplication& Result) const
 	{
 		check(ActivityType == EConcertSyncReplicationActivityType::LeaveReplication);
 		return Payload.GetTypedPayload(Result);
 	}
+	bool GetPayload(FConcertSyncReplicationPayload_Mute& Result) const
+	{
+		check(ActivityType == EConcertSyncReplicationActivityType::Mute);
+		return Payload.GetTypedPayload(Result);
+	}
 
 	CONCERTSYNCCORE_API friend bool operator==(const FConcertSyncReplicationEvent& Left, const FConcertSyncReplicationEvent& Right);
 	friend bool operator!=(const FConcertSyncReplicationEvent&, const FConcertSyncReplicationEvent&) = default;
@@ -151,10 +192,18 @@ struct CONCERTSYNCCORE_API FConcertSyncReplicationActivitySummary : public FConc
 		check(ActivityType == EConcertSyncReplicationActivityType::LeaveReplication);
 		return Payload.GetTypedPayload(Data);
 	}
+	bool GetSummaryData(FConcertSyncReplicationSummary_Mute& Data) const
+	{
+		check(ActivityType == EConcertSyncReplicationActivityType::Mute);
+		return Payload.GetTypedPayload(Data);
+	}
 
 	/** Create this summary from a replication event */
 	static FConcertSyncReplicationActivitySummary CreateSummaryForEvent(const FConcertSyncReplicationEvent& InEvent);
 
+	/** Gets the title for this summary */
+	FText ToDisplayTitle() const;
+
 protected:
 	
 	//~ Begin FConcertSyncActivitySummary Interface
diff --git a/Engine/Plugins/Developer/Concert/ConcertSync/ConcertSyncServer/Source/ConcertSyncServer/Private/Replication/ConcertServerReplicationManager.cpp b/Engine/Plugins/Developer/Concert/ConcertSync/ConcertSyncServer/Source/ConcertSyncServer/Private/Replication/ConcertServerReplicationManager.cpp
--- a/Engine/Plugins/Developer/Concert/ConcertSync/ConcertSyncServer/Source/ConcertSyncServer/Private/Replication/ConcertServerReplicationManager.cpp
+++ b/Engine/Plugins/Developer/Concert/ConcertSync/ConcertSyncServer/Source/ConcertSyncServer/Private/Replication/ConcertServerReplicationManager.cpp
@@ -50,6 +50,8 @@ namespace UE::ConcertSyncServer::Replication
 		
 		Session->OnSessionClientChanged().AddRaw(this, &FConcertServerReplicationManager::OnConnectionChanged);
 		Session->OnTick().AddRaw(this, &FConcertServerReplicationManager::Tick);
+
+		MuteManager.OnMuteRequestApplied().AddRaw(this, &FConcertServerReplicationManager::GenerateMuteActivity);
 	}
 
 	FConcertServerReplicationManager::~FConcertServerReplicationManager()
@@ -296,6 +298,15 @@ namespace UE::ConcertSyncServer::Replication
 		}
 	}
 
+	void FConcertServerReplicationManager::GenerateMuteActivity(const FGuid& EndpointId, const FConcertReplication_ChangeMuteState_Request& Request) const
+	{
+		if (EnumHasAnyFlags(SessionFlags, EConcertSyncSessionFlags::ShouldEnableReplicationActivities))
+		{
+			const FConcertSyncReplicationPayload_Mute MutePayload { Request };
+			ServerWorkspace.ProduceClientMuteReplicationActivity(EndpointId, MutePayload);
+		}
+	}
+
 	void FConcertServerReplicationManager::Tick(IConcertServerSession& InSession, float InDeltaTime)
 	{
 		// TODO UE-190714: We may want to set a time budget for clients.
diff --git a/Engine/Plugins/Developer/Concert/ConcertSync/ConcertSyncServer/Source/ConcertSyncServer/Private/Replication/ConcertServerReplicationManager.h b/Engine/Plugins/Developer/Concert/ConcertSync/ConcertSyncServer/Source/ConcertSyncServer/Private/Replication/ConcertServerReplicationManager.h
--- a/Engine/Plugins/Developer/Concert/ConcertSync/ConcertSyncServer/Source/ConcertSyncServer/Private/Replication/ConcertServerReplicationManager.h
+++ b/Engine/Plugins/Developer/Concert/ConcertSync/ConcertSyncServer/Source/ConcertSyncServer/Private/Replication/ConcertServerReplicationManager.h
@@ -21,7 +21,6 @@
 #include "Templates/Tuple.h"
 #include "Templates/UnrealTemplate.h"
 
-struct FConcertSyncReplicationPayload_LeaveReplication;
 class FConcertServerWorkspace;
 class IConcertClientReplicationBridge;
 class IConcertServerSession;
@@ -29,10 +28,12 @@ class IConcertServerSession;
 enum class EConcertSyncSessionFlags : uint32;
 enum class EConcertQueryClientStreamFlags : uint8;
 
+struct FConcertAuthorityClientInfo;
+struct FConcertReplication_ChangeMuteState_Request;
 struct FConcertReplication_ChangeStream_Response;
 struct FConcertReplication_QueryReplicationInfo_Response;
 struct FConcertReplication_QueryReplicationInfo_Request;
-struct FConcertAuthorityClientInfo;
+struct FConcertSyncReplicationPayload_LeaveReplication;
 
 namespace UE::ConcertSyncCore
 {
@@ -57,7 +58,7 @@ namespace UE::ConcertSyncServer::Replication
 		, public FNoncopyable
 	{
 	public:
-
+		
 		explicit FConcertServerReplicationManager(TSharedRef<IConcertServerSession> InLiveSession, IReplicationWorkspace& InServerWorkspace UE_LIFETIMEBOUND, EConcertSyncSessionFlags InSessionFlags);
 		virtual ~FConcertServerReplicationManager() override;
 
@@ -144,6 +145,9 @@ namespace UE::ConcertSyncServer::Replication
 		void OnClientLeftReplication(const FGuid& EndpointId);
 		/** Calls FConcertServerWorkspace::AddReplicationActivity with the current client state so it can be restored upon re-joining. */
 		void ProduceClientLeftActivity(const FConcertReplicationClient& Client) const;
+
+		/** Generates an activity for a mute request */
+		void GenerateMuteActivity(const FGuid& EndpointId, const FConcertReplication_ChangeMuteState_Request& Request) const;
 		
 		/**
 		 * Ticks all clients which causes clients to process pending data and send it to the corresponding endpoints.
diff --git a/Engine/Plugins/Developer/Concert/ConcertSync/ConcertSyncServer/Source/ConcertSyncServer/Private/Replication/MuteManager.cpp b/Engine/Plugins/Developer/Concert/ConcertSync/ConcertSyncServer/Source/ConcertSyncServer/Private/Replication/MuteManager.cpp
--- a/Engine/Plugins/Developer/Concert/ConcertSync/ConcertSyncServer/Source/ConcertSyncServer/Private/Replication/MuteManager.cpp
+++ b/Engine/Plugins/Developer/Concert/ConcertSync/ConcertSyncServer/Source/ConcertSyncServer/Private/Replication/MuteManager.cpp
@@ -221,7 +221,8 @@ namespace UE::ConcertSyncServer::Replication
 #if UE_BUILD_DEBUG // In debug builds so we don't tank performance
 		checkf(!Response.SyncControl.DoesAtLeastOneObjectLoseSyncControl(), TEXT("Clients are supposed to infer that objects lose sync control. We only tell them which objects gained sync control."));
 #endif
-		
+
+		OnMuteRequestAppliedDelegate.Broadcast(Context.SourceEndpointId, Request);
 		return EConcertSessionResponseCode::Success;
 	}
 
diff --git a/Engine/Plugins/Developer/Concert/ConcertSync/ConcertSyncServer/Source/ConcertSyncServer/Private/Replication/MuteManager.h b/Engine/Plugins/Developer/Concert/ConcertSync/ConcertSyncServer/Source/ConcertSyncServer/Private/Replication/MuteManager.h
--- a/Engine/Plugins/Developer/Concert/ConcertSync/ConcertSyncServer/Source/ConcertSyncServer/Private/Replication/MuteManager.h
+++ b/Engine/Plugins/Developer/Concert/ConcertSync/ConcertSyncServer/Source/ConcertSyncServer/Private/Replication/MuteManager.h
@@ -43,6 +43,9 @@ namespace UE::ConcertSyncServer::Replication
 		DECLARE_DELEGATE_OneParam(FOnMuteStateChangedByClient, const FGuid& ClientId);
 		/** Updates sync control for all clients, sends an update to all clients but ClientId, and returns the sync control to embed into the mute response. */
 		DECLARE_DELEGATE_RetVal_OneParam(FConcertReplication_ChangeSyncControl, FGenerateSyncControlForMuteChange, const FGuid& ClientId);
+		
+		/** Notifies manager about applied event. Generates mute activity. */
+		DECLARE_MULTICAST_DELEGATE_TwoParams(FOnMuteRequestApplied, const FGuid& ClientId, const FConcertReplication_ChangeMuteState_Request& Request);
 
 		FMuteManager(
 			IConcertSession& InSession UE_LIFETIMEBOUND,
@@ -62,6 +65,8 @@ namespace UE::ConcertSyncServer::Replication
 		FOnMuteStateChangedByClient& OnUpdateSyncControlForIndirectMuteChange() { return OnMuteStateChangedDelegate; }
 		FGenerateSyncControlForMuteChange& OnGenerateSyncControlForMuteChange() { return OnGenerateSyncControlForMuteChangeDelegate; }
 		
+		FOnMuteRequestApplied& OnMuteRequestApplied() { return OnMuteRequestAppliedDelegate; }
+		
 	private:
 
 		/** The session to receive requests on. */
@@ -81,6 +86,9 @@ namespace UE::ConcertSyncServer::Replication
 		 * The sync control will only contain RestrictToObjects.
 		 */
 		FGenerateSyncControlForMuteChange OnGenerateSyncControlForMuteChangeDelegate;
+		
+		/** Broadcasts after a mute request has been applied. */
+		FOnMuteRequestApplied OnMuteRequestAppliedDelegate;
 
 		enum class EMuteState : uint8
 		{
diff --git a/Engine/Plugins/Developer/Concert/ConcertSync/ConcertSyncServer/Source/ConcertSyncServer/Private/Replication/ReplicationWorkspace.cpp b/Engine/Plugins/Developer/Concert/ConcertSync/ConcertSyncServer/Source/ConcertSyncServer/Private/Replication/ReplicationWorkspace.cpp
--- a/Engine/Plugins/Developer/Concert/ConcertSync/ConcertSyncServer/Source/ConcertSyncServer/Private/Replication/ReplicationWorkspace.cpp
+++ b/Engine/Plugins/Developer/Concert/ConcertSync/ConcertSyncServer/Source/ConcertSyncServer/Private/Replication/ReplicationWorkspace.cpp
@@ -21,19 +21,12 @@ namespace UE::ConcertSyncServer
 
 	TOptional<int64> FReplicationWorkspace::ProduceClientLeaveReplicationActivity(const FGuid& EndpointId, const FConcertSyncReplicationPayload_LeaveReplication& EventData)
 	{
-		FConcertSyncReplicationActivity Activity;
-		Activity.EndpointId = EndpointId;
-		Activity.EventData.SetPayload(EventData);
-		Activity.EventSummary.SetTypedPayload(FConcertSyncReplicationActivitySummary::CreateSummaryForEvent(Activity.EventData));
-		Activity.bIgnored = ShouldIgnoreClientActivityOnRestoreDelegate.Execute(EndpointId);
-		
-		int64 ActivityId = 0;
-		int64 EventId = 0;
-		const bool bSuccess = Database.AddReplicationActivity(Activity, ActivityId, EventId);
-		// This allows FConcertServerWorkspace to send the activity to other clients
-		OnAddReplicationActivityDelegate.Broadcast(ActivityId, bSuccess);
+		return ProduceActivity(EndpointId, EventData);
+	}
 
-		return bSuccess ? ActivityId : TOptional<int64>{};
+	TOptional<int64> FReplicationWorkspace::ProduceClientMuteReplicationActivity(const FGuid& EndpointId, const FConcertSyncReplicationPayload_Mute& EventData)
+	{
+		return ProduceActivity(EndpointId, EventData);
 	}
 
 	bool FReplicationWorkspace::GetLastLeaveReplicationActivityByClient(const FConcertSessionClientInfo& InClientInfo, FConcertSyncReplicationPayload_LeaveReplication& OutLeaveReplication) const
@@ -81,4 +74,22 @@ namespace UE::ConcertSyncServer
 		const bool bGotEventData = Database.GetReplicationEvent(ActivityId, Event);
 		return bGotEventData && Event.ActivityType == EConcertSyncReplicationActivityType::LeaveReplication && Event.GetPayload(OutLeaveReplication);
 	}
+	
+	template <typename TPayload>
+	TOptional<int64> FReplicationWorkspace::ProduceActivity(const FGuid& EndpointId, const TPayload& EventData)
+	{
+		FConcertSyncReplicationActivity Activity;
+		Activity.EndpointId = EndpointId;
+		Activity.EventData.SetPayload(EventData);
+		Activity.EventSummary.SetTypedPayload(FConcertSyncReplicationActivitySummary::CreateSummaryForEvent(Activity.EventData));
+		Activity.bIgnored = ShouldIgnoreClientActivityOnRestoreDelegate.Execute(EndpointId);
+		
+		int64 ActivityId = 0;
+		int64 EventId = 0;
+		const bool bSuccess = Database.AddReplicationActivity(Activity, ActivityId, EventId);
+		// This allows FConcertServerWorkspace to send the activity to other clients
+		OnAddReplicationActivityDelegate.Broadcast(ActivityId, bSuccess);
+
+		return bSuccess ? ActivityId : TOptional<int64>{};
+	}
 }
diff --git a/Engine/Plugins/Developer/Concert/ConcertSync/ConcertSyncServer/Source/ConcertSyncServer/Private/Replication/ReplicationWorkspace.h b/Engine/Plugins/Developer/Concert/ConcertSync/ConcertSyncServer/Source/ConcertSyncServer/Private/Replication/ReplicationWorkspace.h
--- a/Engine/Plugins/Developer/Concert/ConcertSync/ConcertSyncServer/Source/ConcertSyncServer/Private/Replication/ReplicationWorkspace.h
+++ b/Engine/Plugins/Developer/Concert/ConcertSync/ConcertSyncServer/Source/ConcertSyncServer/Private/Replication/ReplicationWorkspace.h
@@ -37,6 +37,7 @@ namespace UE::ConcertSyncServer
 
 		//~ Begin IReplicationWorkspace Interface
 		virtual TOptional<int64> ProduceClientLeaveReplicationActivity(const FGuid& EndpointId, const FConcertSyncReplicationPayload_LeaveReplication& EventData) override;
+		virtual TOptional<int64> ProduceClientMuteReplicationActivity(const FGuid& EndpointId, const FConcertSyncReplicationPayload_Mute& EventData) override;
 		virtual bool GetLastLeaveReplicationActivityByClient(const FConcertSessionClientInfo& InClientInfo, FConcertSyncReplicationPayload_LeaveReplication& OutLeaveReplication) const override;
 		virtual bool GetLeaveReplicationActivityById(const int64 ActivityId, FConcertSyncReplicationPayload_LeaveReplication& OutLeaveReplication) const override;
 		//~ End IReplicationWorkspace Interface
@@ -51,6 +52,9 @@ namespace UE::ConcertSyncServer
 		const FShouldIgnoreClientActivityOnRestore ShouldIgnoreClientActivityOnRestoreDelegate;
 		
 		FOnAddReplicationActivity OnAddReplicationActivityDelegate;
+
+		template<typename TPayload>
+		TOptional<int64> ProduceActivity(const FGuid& EndpointId, const TPayload& EventData);
 	};
 }
 
diff --git a/Engine/Plugins/Developer/Concert/ConcertSync/ConcertSyncServer/Source/ConcertSyncServer/Public/Replication/IReplicationWorkspace.h b/Engine/Plugins/Developer/Concert/ConcertSync/ConcertSyncServer/Source/ConcertSyncServer/Public/Replication/IReplicationWorkspace.h
--- a/Engine/Plugins/Developer/Concert/ConcertSync/ConcertSyncServer/Source/ConcertSyncServer/Public/Replication/IReplicationWorkspace.h
+++ b/Engine/Plugins/Developer/Concert/ConcertSync/ConcertSyncServer/Source/ConcertSyncServer/Public/Replication/IReplicationWorkspace.h
@@ -5,6 +5,7 @@
 #include "HAL/Platform.h"
 #include "Misc/Optional.h"
 
+struct FConcertSyncReplicationPayload_Mute;
 struct FConcertSessionClientInfo;
 struct FConcertSyncReplicationPayload_LeaveReplication;
 struct FGuid;
@@ -20,11 +21,17 @@ namespace UE::ConcertSyncServer::Replication
 	public:
 		
 		/**
-		 * Creates a replication activity for the client leaving replication if the session has the EConcertSyncSessionFlags::ShouldEnableReplicationActivities flag.
+		 * Creates a replication activity for the client leaving replication
 		 * @return The identifier of the produced activity. Unset if activity insertion failed.
 		 */
 		virtual TOptional<int64> ProduceClientLeaveReplicationActivity(const FGuid& EndpointId, const FConcertSyncReplicationPayload_LeaveReplication& EventData) = 0;
 
+		/**
+		 * Creates a replication activity for the client (un)muting objects in the session.
+		 * @return The identifier of the produced activity. Unset if activity insertion failed.
+		 */
+		virtual TOptional<int64> ProduceClientMuteReplicationActivity(const FGuid& EndpointId, const FConcertSyncReplicationPayload_Mute& EventData) = 0;
+
 		/**
 		 * Gets the last replication leave activity associated for a given client.
 		 * 
diff --git a/Engine/Plugins/Developer/Concert/ConcertSync/ConcertSyncTest/Source/ConcertSyncTest/Private/Replication/Restore/ActivityProduction.spec.cpp b/Engine/Plugins/Developer/Concert/ConcertSync/ConcertSyncTest/Source/ConcertSyncTest/Private/Replication/Restore/ActivityProduction.spec.cpp
--- a/Engine/Plugins/Developer/Concert/ConcertSync/ConcertSyncTest/Source/ConcertSyncTest/Private/Replication/Restore/ActivityProduction.spec.cpp
+++ b/Engine/Plugins/Developer/Concert/ConcertSync/ConcertSyncTest/Source/ConcertSyncTest/Private/Replication/Restore/ActivityProduction.spec.cpp
@@ -75,6 +75,8 @@ namespace UE::ConcertSyncTests::Replication
 		AfterEach([this]
 		{
 			Server.Reset();
+			ObjectReplicator.Reset();
+			WorkspaceMock.Reset();
 		});
 		
 		It("When leaving replication, an activity is produced", [this]
@@ -88,10 +90,69 @@ namespace UE::ConcertSyncTests::Replication
 			ValidateLeaveActivity();
 		});
 	}
+
+
+	BEGIN_DEFINE_SPEC(FMuteReplicationActivitySpec, "Editor.Concert.Replication.RestoreContent.Activity", EAutomationTestFlags::EditorContext | EAutomationTestFlags::EngineFilter)
+		/** Detects calls into the workspace */
+		TSharedPtr<FReplicationWorkspaceCallInterceptorMock> WorkspaceMock;
+	
+		TUniquePtr<FObjectTestReplicator> ObjectReplicator;
+		TUniquePtr<FReplicationServer> Server;
+		FReplicationClient* Client = nullptr;
+	
+		FGuid SenderStreamId = FGuid::NewGuid();
+	END_DEFINE_SPEC(FMuteReplicationActivitySpec);
+
+	/** This tests that when a client leaves replication, an activity containing their registered streams and authority is produced. */
+	void FMuteReplicationActivitySpec::Define()
+	{
+		BeforeEach([this]
+		{
+			WorkspaceMock = MakeShared<FReplicationWorkspaceCallInterceptorMock>();
+			ObjectReplicator = MakeUnique<FObjectTestReplicator>();
+			Server = MakeUnique<FReplicationServer>(*this, EConcertSyncSessionFlags::Default_MultiUserSession, WorkspaceMock.ToSharedRef());
+			Client = &Server->ConnectClient();
+			
+			Client->JoinReplication(ObjectReplicator->CreateSenderArgs(SenderStreamId));
+			Client->GetClientReplicationManager().TakeAuthorityOver({ ObjectReplicator->TestObject });
+		});
+		AfterEach([this]
+		{
+			Server.Reset();
+			ObjectReplicator.Reset();
+			WorkspaceMock.Reset();
+		});
+		
+		It("When muting replication, an activity is produced", [this]
+		{
+			bool bMuteSuccess = false;
+			Client->GetClientReplicationManager()
+				.MuteObjects({ ObjectReplicator->TestObject })
+				.Next([&bMuteSuccess](FConcertReplication_ChangeMuteState_Response&& Response)
+				{
+					bMuteSuccess = Response.IsSuccess();
+				});
+			
+			TestTrue(TEXT("Mute success"), bMuteSuccess);
+			if (!WorkspaceMock->LastCall_ProduceClientMuteReplicationActivity)
+			{
+				AddError(TEXT("No activity produced"));
+				return;
+			}
+			
+			const auto[EndpointId, EventData] = *WorkspaceMock->LastCall_ProduceClientMuteReplicationActivity;
+			const FConcertReplication_ChangeMuteState_Request& Request = EventData.Request;
+			TestEqual(TEXT("EndpointId"), EndpointId, Client->GetEndpointId());
+			TestEqual(TEXT("ObjectsToMute.Num()"), Request.ObjectsToMute.Num(), 1);
+			TestEqual(TEXT("ObjectsToUnmute.Num()"), Request.ObjectsToUnmute.Num(), 0);
+			TestTrue(TEXT("ObjectsToMute.Contains(TestObject)"), Request.ObjectsToMute.Contains(ObjectReplicator->TestObject));
+		});
+	}
 	
 	BEGIN_DEFINE_SPEC(FNoReplicationActivitiesSpec, "Editor.Concert.Replication.RestoreContent.Activity", EAutomationTestFlags::EditorContext | EAutomationTestFlags::EngineFilter)
 		/** Detects calls into the workspace */
 		TSharedPtr<FReplicationWorkspaceCallInterceptorMock> WorkspaceMock;
+		TUniquePtr<FObjectTestReplicator> ObjectReplicator;
 	
 		TUniquePtr<FReplicationServer> Server;
 		FReplicationClient* Client = nullptr;
@@ -103,14 +164,16 @@ namespace UE::ConcertSyncTests::Replication
 		BeforeEach([this]
 		{
 			WorkspaceMock = MakeShared<FReplicationWorkspaceCallInterceptorMock>();
+			ObjectReplicator = MakeUnique<FObjectTestReplicator>();
 			Server = MakeUnique<FReplicationServer>(*this, EConcertSyncSessionFlags::Default_MultiUserSession & ~EConcertSyncSessionFlags::ShouldEnableReplicationActivities, WorkspaceMock.ToSharedRef());
 			Client = &Server->ConnectClient();
 
-			Client->JoinReplicationAsListener({});
+			Client->JoinReplication(ObjectReplicator->CreateSenderArgs());
 		});
 		AfterEach([this]
 		{
 			Server.Reset();
+			ObjectReplicator.Reset();
 			WorkspaceMock.Reset();
 		});
 		
@@ -119,5 +182,18 @@ namespace UE::ConcertSyncTests::Replication
 			Client->LeaveReplication();
 			TestFalse(TEXT("Activity produced"), WorkspaceMock->LastCall_ProduceClientLeaveReplicationActivity.IsSet());
 		});
+		It("When muting, no activity is produced", [this]
+		{
+			bool bMuteSuccess = false;
+			Client->GetClientReplicationManager()
+				.MuteObjects({ ObjectReplicator->TestObject })
+				.Next([&bMuteSuccess](FConcertReplication_ChangeMuteState_Response&& Response)
+				{
+					bMuteSuccess = Response.IsSuccess();
+				});
+			
+			TestTrue(TEXT("Mute success"), bMuteSuccess);
+			TestFalse(TEXT("Activity produced"), WorkspaceMock->LastCall_ProduceClientMuteReplicationActivity.IsSet());
+		});
 	}
 }
diff --git a/Engine/Plugins/Developer/Concert/ConcertSync/ConcertSyncTest/Source/ConcertSyncTest/Private/Replication/Util/Mocks/ReplicationWorkspaceCallInterceptorMock.h b/Engine/Plugins/Developer/Concert/ConcertSync/ConcertSyncTest/Source/ConcertSyncTest/Private/Replication/Util/Mocks/ReplicationWorkspaceCallInterceptorMock.h
--- a/Engine/Plugins/Developer/Concert/ConcertSync/ConcertSyncTest/Source/ConcertSyncTest/Private/Replication/Util/Mocks/ReplicationWorkspaceCallInterceptorMock.h
+++ b/Engine/Plugins/Developer/Concert/ConcertSync/ConcertSyncTest/Source/ConcertSyncTest/Private/Replication/Util/Mocks/ReplicationWorkspaceCallInterceptorMock.h
@@ -17,6 +17,8 @@ namespace UE::ConcertSyncTests::Replication
 
 		/** Arguments of last ProduceClientLeaveReplicationActivity call. */
 		TOptional<TTuple<FGuid, FConcertSyncReplicationPayload_LeaveReplication>> LastCall_ProduceClientLeaveReplicationActivity;
+		/** Arguments of last ProduceClientMuteReplicationActivity call. */
+		TOptional<TTuple<FGuid, FConcertSyncReplicationPayload_Mute>> LastCall_ProduceClientMuteReplicationActivity;
 		/** Arguments of last GetLastLeaveReplicationActivityByClient call. */
 		mutable TOptional<TTuple<FConcertSessionClientInfo>> LastCall_GetLastLeaveReplicationActivityByClient;
 		/** Arguments of last GetLastLeaveReplicationActivityByClient call. */
@@ -24,6 +26,8 @@ namespace UE::ConcertSyncTests::Replication
 
 		/** The result to return in ProduceClientLeaveReplicationActivity. */
 		TOptional<int64> ReturnResult_ProduceClientLeaveReplicationActivity = 0;
+		/** The result to return in ProduceClientLeaveReplicationActivity. */
+		TOptional<int64> ReturnResult_ProduceClientMuteReplicationActivity = 0;
 		/** The result to return in GetLastLeaveReplicationActivityByClient. */
 		TOptional<FConcertSyncReplicationPayload_LeaveReplication> ReturnResult_GetLastLeaveReplicationActivityByClient;
 		/** The result to return in GetLastLeaveReplicationActivityByClient. */
@@ -35,6 +39,12 @@ namespace UE::ConcertSyncTests::Replication
 			LastCall_ProduceClientLeaveReplicationActivity = MakeTuple(EndpointId, EventData);
 			return ReturnResult_ProduceClientLeaveReplicationActivity;
 		}
+		virtual TOptional<int64> ProduceClientMuteReplicationActivity(const FGuid& EndpointId, const FConcertSyncReplicationPayload_Mute& EventData) override
+		{
+			LastCall_ProduceClientMuteReplicationActivity = MakeTuple(EndpointId, EventData);
+			return ReturnResult_ProduceClientMuteReplicationActivity;
+		}
+
 		virtual bool GetLastLeaveReplicationActivityByClient(const FConcertSessionClientInfo& InClientInfo, FConcertSyncReplicationPayload_LeaveReplication& OutLeaveReplication) const override
 		{
 			LastCall_GetLastLeaveReplicationActivityByClient = MakeTuple(InClientInfo);
diff --git a/Engine/Plugins/Developer/Concert/ConcertSync/ConcertSyncTest/Source/ConcertSyncTest/Private/Replication/Util/Mocks/ReplicationWorkspaceEmptyMock.h b/Engine/Plugins/Developer/Concert/ConcertSync/ConcertSyncTest/Source/ConcertSyncTest/Private/Replication/Util/Mocks/ReplicationWorkspaceEmptyMock.h
--- a/Engine/Plugins/Developer/Concert/ConcertSync/ConcertSyncTest/Source/ConcertSyncTest/Private/Replication/Util/Mocks/ReplicationWorkspaceEmptyMock.h
+++ b/Engine/Plugins/Developer/Concert/ConcertSync/ConcertSyncTest/Source/ConcertSyncTest/Private/Replication/Util/Mocks/ReplicationWorkspaceEmptyMock.h
@@ -12,6 +12,7 @@ namespace UE::ConcertSyncTests::Replication
 		
 		//~ Begin IReplicationWorkspace Interface
 		virtual TOptional<int64> ProduceClientLeaveReplicationActivity(const FGuid&, const FConcertSyncReplicationPayload_LeaveReplication&) override { return {}; }
+		virtual TOptional<int64> ProduceClientMuteReplicationActivity(const FGuid& EndpointId, const FConcertSyncReplicationPayload_Mute& EventData) override { return {}; }
 		virtual bool GetLastLeaveReplicationActivityByClient(const FConcertSessionClientInfo& InClientInfo, FConcertSyncReplicationPayload_LeaveReplication& OutLeaveReplication) const override { return false; }
 		virtual bool GetLeaveReplicationActivityById(const int64 ActivityId, FConcertSyncReplicationPayload_LeaveReplication& OutLeaveReplication) const override { return false; }
 		//~ End IReplicationWorkspace Interface
diff --git a/Engine/Plugins/Developer/Concert/ConcertUI/ConcertSharedSlate/Source/ConcertSharedSlate/Private/Session/Activity/SessionActivityUtils.cpp b/Engine/Plugins/Developer/Concert/ConcertUI/ConcertSharedSlate/Source/ConcertSharedSlate/Private/Session/Activity/SessionActivityUtils.cpp
--- a/Engine/Plugins/Developer/Concert/ConcertUI/ConcertSharedSlate/Source/ConcertSharedSlate/Private/Session/Activity/SessionActivityUtils.cpp
+++ b/Engine/Plugins/Developer/Concert/ConcertUI/ConcertSharedSlate/Source/ConcertSharedSlate/Private/Session/Activity/SessionActivityUtils.cpp
@@ -59,7 +59,7 @@ FText UE::ConcertSharedSlate::Private::GetOperationName(const FConcertSessionAct
 
 	if (const FConcertSyncReplicationActivitySummary* ActivitySummary = Activity.ActivitySummary.Cast<FConcertSyncReplicationActivitySummary>())
 	{
-		return LOCTEXT("LeftReplication", "Left Replication");
+		return ActivitySummary->ToDisplayTitle();
 	}
 
 	return FText::GetEmpty();
