Fix copy-paste typo in FConcertSyncSessionDatabase::EnumerateReplicationActivities

diff --git a/Engine/Plugins/Developer/Concert/ConcertSync/ConcertSyncCore/Source/ConcertSyncCore/Private/ConcertSyncSessionDatabase.cpp b/Engine/Plugins/Developer/Concert/ConcertSync/ConcertSyncCore/Source/ConcertSyncCore/Private/ConcertSyncSessionDatabase.cpp
--- a/Engine/Plugins/Developer/Concert/ConcertSync/ConcertSyncCore/Source/ConcertSyncCore/Private/ConcertSyncSessionDatabase.cpp
+++ b/Engine/Plugins/Developer/Concert/ConcertSync/ConcertSyncCore/Source/ConcertSyncCore/Private/ConcertSyncSessionDatabase.cpp
@@ -2132,7 +2132,7 @@ bool FConcertSyncSessionDatabase::EnumeratePackageActivities(FIteratePackageActi
 
 bool FConcertSyncSessionDatabase::EnumerateReplicationActivities(TFunctionRef<bool(FConcertSyncReplicationActivity&&)> InCallback) const
 {
-	return Statements->GetAllActivityDataForEventType(EConcertSyncActivityEventType::Transaction, [this, &InCallback](const int64 InActivityId, const EConcertSyncActivityFlags InFlags, const FGuid& InEndpointId, const FDateTime InEventTime, const int64 InEventId, FConcertSessionSerializedPayload&& InEventSummary)
+	return Statements->GetAllActivityDataForEventType(EConcertSyncActivityEventType::Replication, [this, &InCallback](const int64 InActivityId, const EConcertSyncActivityFlags InFlags, const FGuid& InEndpointId, const FDateTime InEventTime, const int64 InEventId, FConcertSessionSerializedPayload&& InEventSummary)
 	{
 		FConcertSyncReplicationActivity ReplicationActivity;
 		ReplicationActivity.ActivityId = InActivityId;
