Fix duplicate loca keys

diff --git a/Engine/Plugins/Developer/Concert/ConcertSync/ConcertSyncCore/Source/ConcertSyncCore/Private/Replication/Messages/ReplicationActivity.cpp b/Engine/Plugins/Developer/Concert/ConcertSync/ConcertSyncCore/Source/ConcertSyncCore/Private/Replication/Messages/ReplicationActivity.cpp
--- a/Engine/Plugins/Developer/Concert/ConcertSync/ConcertSyncCore/Source/ConcertSyncCore/Private/Replication/Messages/ReplicationActivity.cpp
+++ b/Engine/Plugins/Developer/Concert/ConcertSync/ConcertSyncCore/Source/ConcertSyncCore/Private/Replication/Messages/ReplicationActivity.cpp
@@ -101,7 +101,7 @@ FText FConcertSyncReplicationActivitySummary::ToDisplayTitle() const
 				return LOCTEXT("Title.Mute.Resume", "Resume replication");
 			}
 
-			return LOCTEXT("Title.Mute.Resume", "Pause / Resume (empty)");
+			return LOCTEXT("Title.Mute.Empty", "Pause / Resume (empty)");
 		}
 	
 	case EConcertSyncReplicationActivityType::None: [[fallthrough]];
