Fix struct not being initialized properly

diff --git a/Engine/Plugins/Developer/Concert/ConcertSync/ConcertSyncCore/Source/ConcertSyncCore/Public/Replication/Messages/ReplicationActivity.h b/Engine/Plugins/Developer/Concert/ConcertSync/ConcertSyncCore/Source/ConcertSyncCore/Public/Replication/Messages/ReplicationActivity.h
--- a/Engine/Plugins/Developer/Concert/ConcertSync/ConcertSyncCore/Source/ConcertSyncCore/Public/Replication/Messages/ReplicationActivity.h
+++ b/Engine/Plugins/Developer/Concert/ConcertSync/ConcertSyncCore/Source/ConcertSyncCore/Public/Replication/Messages/ReplicationActivity.h
@@ -181,7 +181,7 @@ struct CONCERTSYNCCORE_API FConcertSyncReplicationActivitySummary : public FConc
 
 	/** The type of replication event we summarize */
 	UPROPERTY()
-	EConcertSyncReplicationActivityType ActivityType;
+	EConcertSyncReplicationActivityType ActivityType { EConcertSyncReplicationActivityType::None };
 
 	/** The summary data. The underlying type depends on ActivityType. */
 	UPROPERTY()
