Remove EConcertReplicationPutStateFlags::ResetMuteState

diff --git a/Engine/Plugins/Developer/Concert/ConcertSync/ConcertSyncCore/Source/ConcertSyncCore/Public/Replication/Messages/PutState.h b/Engine/Plugins/Developer/Concert/ConcertSync/ConcertSyncCore/Source/ConcertSyncCore/Public/Replication/Messages/PutState.h
--- a/Engine/Plugins/Developer/Concert/ConcertSync/ConcertSyncCore/Source/ConcertSyncCore/Public/Replication/Messages/PutState.h
+++ b/Engine/Plugins/Developer/Concert/ConcertSync/ConcertSyncCore/Source/ConcertSyncCore/Public/Replication/Messages/PutState.h
@@ -15,9 +15,6 @@ enum class EConcertReplicationPutStateFlags : uint8
 	/** If another client has disconnected between the time the request was sent and received by the client, do not fail the request. */
 	SkipDisconnectedClients = 1 << 0,
 
-	/** Prior to applying the FConcertReplication_PutState_Request::MuteChange, clear all mute state. */
-	ResetMuteState = 1 << 1,
-
 	Default = SkipDisconnectedClients
 };
 ENUM_CLASS_FLAGS(EConcertReplicationPutStateFlags);
