Make implementing new replication activities and summary operators easier / more reusable.

diff --git a/Engine/Plugins/Developer/Concert/ConcertSync/ConcertSyncCore/Source/ConcertSyncCore/Private/Replication/Messages/ReplicationActivity.cpp b/Engine/Plugins/Developer/Concert/ConcertSync/ConcertSyncCore/Source/ConcertSyncCore/Private/Replication/Messages/ReplicationActivity.cpp
--- a/Engine/Plugins/Developer/Concert/ConcertSync/ConcertSyncCore/Source/ConcertSyncCore/Private/Replication/Messages/ReplicationActivity.cpp
+++ b/Engine/Plugins/Developer/Concert/ConcertSync/ConcertSyncCore/Source/ConcertSyncCore/Private/Replication/Messages/ReplicationActivity.cpp
@@ -4,35 +4,84 @@
 
 #include "Internationalization/Text.h"
 
+#include <type_traits>
+
 #define LOCTEXT_NAMESPACE "ReplicationActivity"
 
+namespace UE::ConcertSyncCore::ReplicationActivity::Private
+{
+	template<typename TPayload>
+	concept CPayloadRetrievable = requires (const FConcertSyncReplicationEvent& Event, TPayload&& Payload){ Event.GetPayload(Payload); };
+	
+	template<CPayloadRetrievable TPayload> 
+	static bool IsReplicationPayloadEqual(const FConcertSyncReplicationEvent& Left, const FConcertSyncReplicationEvent& Right)
+	{
+		TPayload LeftContent;
+		TPayload RightContent;
+		const bool bLeftSucceeded = Left.GetPayload(LeftContent);
+		const bool bRightSucceeded = Right.GetPayload(RightContent);
+		return bLeftSucceeded && bRightSucceeded && LeftContent == RightContent;
+	}
+
+	template<CPayloadRetrievable TPayload, typename TSummary> requires std::is_constructible_v<TSummary, const TPayload&>
+	static void FillSummary(const FConcertSyncReplicationEvent& InEvent, FConcertSyncReplicationActivitySummary& Summary)
+	{
+		TPayload EventPayload;
+		InEvent.GetPayload(EventPayload);
+			
+		const TSummary SummaryPayload(EventPayload);
+		Summary.Payload.SetTypedPayload(SummaryPayload);
+	}
+	
+	static FText GetDisplayText(const FConcertSyncReplicationSummary_Mute& SummaryData)
+	{
+		const bool bHasMuted = !SummaryData.Request.ObjectsToMute.IsEmpty();
+		const bool bHasUnmuted = !SummaryData.Request.ObjectsToUnmute.IsEmpty();
+		if (bHasMuted && bHasUnmuted)
+		{
+			return LOCTEXT("Title.Mute.PauseAndResume", "Pause & Resume");
+		}
+
+		if (bHasMuted)
+		{
+			return LOCTEXT("Title.Mute.Pause", "Pause replication");
+		}
+			
+		if (bHasUnmuted)
+		{
+			return LOCTEXT("Title.Mute.Resume", "Resume replication");
+		}
+
+		return LOCTEXT("Title.Mute.Empty", "Pause / Resume (empty)");
+	}
+
+	template<typename TSummary>
+	concept CSummaryGettable = requires (const FConcertSyncReplicationActivitySummary& Summary, TSummary& Payload) { Summary.GetSummaryData(Payload); };
+
+	template<CSummaryGettable TSummary> 
+	static FText GetSummaryText(const FConcertSyncReplicationActivitySummary& Summary, FText NoSummaryText)
+	{
+		TSummary SummaryData;
+		return Summary.GetSummaryData(SummaryData)
+			? GetDisplayText(SummaryData) // Take advantage of overloading
+			: NoSummaryText;
+	}
+}
+
 bool operator==(const FConcertSyncReplicationEvent& Left, const FConcertSyncReplicationEvent& Right)
 {
 	if (Left.ActivityType != Right.ActivityType)
 	{
 		return false;
 	}
-
+	
+	using namespace UE::ConcertSyncCore::ReplicationActivity::Private;
 	static_assert(static_cast<int32>(EConcertSyncReplicationActivityType::Count) == 3, "If you added an EConcertSyncReplicationActivityType entry, update this switch");
 	switch (Left.ActivityType)
 	{
 	case EConcertSyncReplicationActivityType::None: return true;
-	case EConcertSyncReplicationActivityType::LeaveReplication:
-		{
-			FConcertSyncReplicationPayload_LeaveReplication LeftContent;
-			FConcertSyncReplicationPayload_LeaveReplication RightContent;
-			const bool bLeftSucceeded = Left.GetPayload(LeftContent);
-			const bool bRightSucceeded = Right.GetPayload(RightContent);
-			return bLeftSucceeded && bRightSucceeded && LeftContent == RightContent;
-		}
-	case EConcertSyncReplicationActivityType::Mute:
-		{
-			FConcertSyncReplicationPayload_Mute LeftContent;
-			FConcertSyncReplicationPayload_Mute RightContent;
-			const bool bLeftSucceeded = Left.GetPayload(LeftContent);
-			const bool bRightSucceeded = Right.GetPayload(RightContent);
-			return bLeftSucceeded && bRightSucceeded && LeftContent == RightContent;
-		}
+	case EConcertSyncReplicationActivityType::LeaveReplication: return IsReplicationPayloadEqual<FConcertSyncReplicationPayload_LeaveReplication>(Left, Right);
+	case EConcertSyncReplicationActivityType::Mute: return IsReplicationPayloadEqual<FConcertSyncReplicationPayload_Mute>(Left, Right);
 	default: checkNoEntry(); return false;
 	}
 }
@@ -42,26 +91,15 @@ FConcertSyncReplicationActivitySummary FConcertSyncReplicationActivitySummary::C
 	FConcertSyncReplicationActivitySummary Summary;
 	Summary.ActivityType = InEvent.ActivityType;
 
+	using namespace UE::ConcertSyncCore::ReplicationActivity::Private;
 	static_assert(static_cast<uint8>(EConcertSyncReplicationActivityType::Count) == 3, "If you changed the enum entries, update this switch");
 	switch (Summary.ActivityType)
 	{
 	case EConcertSyncReplicationActivityType::LeaveReplication:
-		{
-			FConcertSyncReplicationPayload_LeaveReplication LeaveReplicationData;
-			InEvent.GetPayload(LeaveReplicationData);
-			
-			const FConcertSyncReplicationSummary_LeaveReplication Content { LeaveReplicationData.OwnedObjects };
-			Summary.Payload.SetTypedPayload(Content);
-		}
+		FillSummary<FConcertSyncReplicationPayload_LeaveReplication, FConcertSyncReplicationSummary_LeaveReplication>(InEvent, Summary);
 		break;
 	case EConcertSyncReplicationActivityType::Mute:
-		{
-			FConcertSyncReplicationPayload_Mute MuteData;
-			InEvent.GetPayload(MuteData);
-			
-			const FConcertSyncReplicationSummary_Mute Content { MuteData.Request };
-			Summary.Payload.SetTypedPayload(Content);
-		}
+		FillSummary<FConcertSyncReplicationPayload_Mute, FConcertSyncReplicationSummary_Mute>(InEvent, Summary);
 		break;
 	case EConcertSyncReplicationActivityType::None: [[fallthrough]];
 	default: checkNoEntry(); break;
@@ -72,38 +110,16 @@ FConcertSyncReplicationActivitySummary FConcertSyncReplicationActivitySummary::C
 
 FText FConcertSyncReplicationActivitySummary::ToDisplayTitle() const
 {
+	using namespace UE::ConcertSyncCore::ReplicationActivity::Private;
 	static_assert(static_cast<uint8>(EConcertSyncReplicationActivityType::Count) == 3, "If you changed the enum entries, update this switch");
 	switch (ActivityType)
 	{
 	case EConcertSyncReplicationActivityType::LeaveReplication: return LOCTEXT("Title.LeftReplication", "Left Replication");
 	case EConcertSyncReplicationActivityType::Mute:
-		{
-			FConcertSyncReplicationSummary_Mute SummaryData;
-			if (!GetSummaryData(SummaryData))
-			{
-				return LOCTEXT("Title.Mute.FailedToGetData", "Pause / Resume");
-			}
-
-			const bool bHasMuted = !SummaryData.Request.ObjectsToMute.IsEmpty();
-			const bool bHasUnmuted = !SummaryData.Request.ObjectsToUnmute.IsEmpty();
-			if (bHasMuted && bHasUnmuted)
-			{
-				return LOCTEXT("Title.Mute.PauseAndResume", "Pause & Resume");
-			}
-
-			if (bHasMuted)
-			{
-				return LOCTEXT("Title.Mute.Pause", "Pause replication");
-			}
-			
-			if (bHasUnmuted)
-			{
-				return LOCTEXT("Title.Mute.Resume", "Resume replication");
-			}
-
-			return LOCTEXT("Title.Mute.Empty", "Pause / Resume (empty)");
-		}
-	
+		return GetSummaryText<FConcertSyncReplicationSummary_Mute>(
+			*this,
+			 LOCTEXT("Title.Mute.FailedToGetData", "Pause / Resume")
+			 );
 	case EConcertSyncReplicationActivityType::None: [[fallthrough]];
 	default: checkNoEntry(); return FText::GetEmpty();
 	}
diff --git a/Engine/Plugins/Developer/Concert/ConcertSync/ConcertSyncCore/Source/ConcertSyncCore/Public/Replication/Messages/ReplicationActivity.h b/Engine/Plugins/Developer/Concert/ConcertSync/ConcertSyncCore/Source/ConcertSyncCore/Public/Replication/Messages/ReplicationActivity.h
--- a/Engine/Plugins/Developer/Concert/ConcertSync/ConcertSyncCore/Source/ConcertSyncCore/Public/Replication/Messages/ReplicationActivity.h
+++ b/Engine/Plugins/Developer/Concert/ConcertSync/ConcertSyncCore/Source/ConcertSyncCore/Public/Replication/Messages/ReplicationActivity.h
@@ -53,6 +53,11 @@ USTRUCT()
 struct FConcertSyncReplicationSummary_LeaveReplication
 {
 	GENERATED_BODY()
+	
+	FConcertSyncReplicationSummary_LeaveReplication() = default;
+	explicit FConcertSyncReplicationSummary_LeaveReplication(const FConcertSyncReplicationPayload_LeaveReplication& Event)
+		: OwnedObjects(Event.OwnedObjects)
+	{}
 
 	/** The objects the client had authority over when they left. */
 	UPROPERTY()
@@ -78,6 +83,11 @@ USTRUCT()
 struct FConcertSyncReplicationSummary_Mute
 {
 	GENERATED_BODY()
+
+	FConcertSyncReplicationSummary_Mute() = default;
+	explicit FConcertSyncReplicationSummary_Mute(const FConcertSyncReplicationPayload_Mute& Event)
+		: Request(Event.Request)
+	{}
 	
 	/** The request that changed mute state */
 	UPROPERTY()
