Fix typo that caused client content to be incorrectly restored.
Instead of getting client content by activity ID, it was being gotten by event ID.

diff --git a/Engine/Plugins/Developer/Concert/ConcertSync/ConcertSyncCore/Source/ConcertSyncCore/Private/ConcertSyncSessionDatabase.cpp b/Engine/Plugins/Developer/Concert/ConcertSync/ConcertSyncCore/Source/ConcertSyncCore/Private/ConcertSyncSessionDatabase.cpp
--- a/Engine/Plugins/Developer/Concert/ConcertSync/ConcertSyncCore/Source/ConcertSyncCore/Private/ConcertSyncSessionDatabase.cpp
+++ b/Engine/Plugins/Developer/Concert/ConcertSync/ConcertSyncCore/Source/ConcertSyncCore/Private/ConcertSyncSessionDatabase.cpp
@@ -1969,6 +1969,12 @@ bool FConcertSyncSessionDatabase::GetReplicationActivity(const int64 InActivityI
 		&& GetReplicationEvent(OutReplicationActivity.EventId, OutReplicationActivity.EventData);
 }
 
+bool FConcertSyncSessionDatabase::GetReplicationActivityForEvent(const int64 InEventId, FConcertSyncReplicationActivity& OutReplicationActivity) const
+{
+	return GetActivityForEvent(InEventId, EConcertSyncActivityEventType::Replication, OutReplicationActivity)
+		&& GetReplicationEvent(OutReplicationActivity.EventId, OutReplicationActivity.EventData);
+}
+
 bool FConcertSyncSessionDatabase::GetActivityEventType(const int64 InActivityId, EConcertSyncActivityEventType& OutEventType) const
 {
 	return Statements->GetActivityEventTypeForId(InActivityId, OutEventType);
diff --git a/Engine/Plugins/Developer/Concert/ConcertSync/ConcertSyncCore/Source/ConcertSyncCore/Public/ConcertSyncSessionDatabase.h b/Engine/Plugins/Developer/Concert/ConcertSync/ConcertSyncCore/Source/ConcertSyncCore/Public/ConcertSyncSessionDatabase.h
--- a/Engine/Plugins/Developer/Concert/ConcertSync/ConcertSyncCore/Source/ConcertSyncCore/Public/ConcertSyncSessionDatabase.h
+++ b/Engine/Plugins/Developer/Concert/ConcertSync/ConcertSyncCore/Source/ConcertSyncCore/Public/ConcertSyncSessionDatabase.h
@@ -291,6 +291,16 @@ public:
 	 * @return True if the replication activity was found, false otherwise.
 	 */
 	bool GetReplicationActivity(const int64 InActivityId, FConcertSyncReplicationActivity& OutReplicationActivity) const;
+
+	/**
+	 * Get a replication activity from this database by its event ID
+	 *
+	 * @param InEventId					The ID of the event that is saved in a replication activity.
+	 * @param OutReplicationActivity	The replication activity to populate with the result.
+	 *
+	 * @return True if the replication activity was found, false otherwise.
+	 */
+	bool GetReplicationActivityForEvent(const int64 InEventId, FConcertSyncReplicationActivity& OutReplicationActivity) const;
 	
 	/**
 	 * Get the type of an activity in this database.
diff --git a/Engine/Plugins/Developer/Concert/ConcertSync/ConcertSyncServer/Source/ConcertSyncServer/Private/Replication/ReplicationWorkspace.cpp b/Engine/Plugins/Developer/Concert/ConcertSync/ConcertSyncServer/Source/ConcertSyncServer/Private/Replication/ReplicationWorkspace.cpp
--- a/Engine/Plugins/Developer/Concert/ConcertSync/ConcertSyncServer/Source/ConcertSyncServer/Private/Replication/ReplicationWorkspace.cpp
+++ b/Engine/Plugins/Developer/Concert/ConcertSync/ConcertSyncServer/Source/ConcertSyncServer/Private/Replication/ReplicationWorkspace.cpp
@@ -2,6 +2,7 @@
 
 #include "ReplicationWorkspace.h"
 
+#include "ConcertLogGlobal.h"
 #include "ConcertSyncSessionDatabase.h"
 #include "Replication/Messages/ReplicationActivity.h"
 
@@ -68,7 +69,7 @@ namespace UE::ConcertSyncServer
 			NewestEventId = FMath::Max(NewestEventId, EventId);
 		}
 
-		const bool bGotData = NewestEventId != INDEX_NONE && Database.GetReplicationActivity(NewestEventId, OutActivity);
+		const bool bGotData = NewestEventId != INDEX_NONE && Database.GetReplicationActivityForEvent(NewestEventId, OutActivity);
 		return bGotData;
 	}
 
