Extract IReplicationGroundTruth into its own file and move it to the UE::ConcertSyncCore::Replication namespace

diff --git a/Engine/Plugins/Developer/Concert/ConcertApp/MultiUserClient/Source/MultiUserClient/Private/Replication/Misc/GlobalAuthorityCache.h b/Engine/Plugins/Developer/Concert/ConcertApp/MultiUserClient/Source/MultiUserClient/Private/Replication/Misc/GlobalAuthorityCache.h
--- a/Engine/Plugins/Developer/Concert/ConcertApp/MultiUserClient/Source/MultiUserClient/Private/Replication/Misc/GlobalAuthorityCache.h
+++ b/Engine/Plugins/Developer/Concert/ConcertApp/MultiUserClient/Source/MultiUserClient/Private/Replication/Misc/GlobalAuthorityCache.h
@@ -5,6 +5,7 @@
 #include "Replication/Client/Online/RemoteClient.h"
 
 #include "Replication/AuthorityConflictSharedUtils.h"
+#include "Replication/Misc/IReplicationGroundTruth.h"
 
 #include "Delegates/Delegate.h"
 #include "Misc/ObjectPathHierarchy.h"
@@ -27,7 +28,7 @@ namespace UE::MultiUserClient::Replication
 	 */
 	class FGlobalAuthorityCache
 		: public FNoncopyable
-		, ConcertSyncCore::Replication::AuthorityConflictUtils::IReplicationGroundTruth
+		, ConcertSyncCore::Replication::IReplicationGroundTruth
 	{
 	public:
 
diff --git a/Engine/Plugins/Developer/Concert/ConcertSync/ConcertSyncCore/Source/ConcertSyncCore/Private/Replication/AuthorityConflictSharedUtils.cpp b/Engine/Plugins/Developer/Concert/ConcertSync/ConcertSyncCore/Source/ConcertSyncCore/Private/Replication/AuthorityConflictSharedUtils.cpp
--- a/Engine/Plugins/Developer/Concert/ConcertSync/ConcertSyncCore/Source/ConcertSyncCore/Private/Replication/AuthorityConflictSharedUtils.cpp
+++ b/Engine/Plugins/Developer/Concert/ConcertSync/ConcertSyncCore/Source/ConcertSyncCore/Private/Replication/AuthorityConflictSharedUtils.cpp
@@ -5,6 +5,7 @@
 #include "Replication/Data/ReplicationStream.h"
 #include "Replication/Messages/ChangeAuthority.h"
 #include "Replication/Messages/ChangeStream.h"
+#include "Replication/Misc/IReplicationGroundTruth.h"
 
 namespace UE::ConcertSyncCore::Replication::AuthorityConflictUtils
 {
diff --git a/Engine/Plugins/Developer/Concert/ConcertSync/ConcertSyncCore/Source/ConcertSyncCore/Private/Replication/Misc/StreamAndAuthorityPredictionUtils.cpp b/Engine/Plugins/Developer/Concert/ConcertSync/ConcertSyncCore/Source/ConcertSyncCore/Private/Replication/Misc/StreamAndAuthorityPredictionUtils.cpp
new file mode 100644
--- /dev/null
+++ b/Engine/Plugins/Developer/Concert/ConcertSync/ConcertSyncCore/Source/ConcertSyncCore/Private/Replication/Misc/StreamAndAuthorityPredictionUtils.cpp
@@ -0,0 +1,13 @@
+﻿// Copyright Epic Games, Inc. All Rights Reserved.
+
+#include "Replication/Misc/StreamAndAuthorityPredictionUtils.h"
+
+#include "ConcertMessageData.h"
+
+namespace UE::ConcertSyncCore::Replication
+{
+	bool AreLogicallySameClients(const FConcertClientInfo& First, const FConcertClientInfo& Second)
+	{
+		return First.DisplayName == Second.DisplayName && First.DeviceName == Second.DeviceName;
+	}
+}
diff --git a/Engine/Plugins/Developer/Concert/ConcertSync/ConcertSyncCore/Source/ConcertSyncCore/Public/Replication/AuthorityConflictSharedUtils.h b/Engine/Plugins/Developer/Concert/ConcertSync/ConcertSyncCore/Source/ConcertSyncCore/Public/Replication/AuthorityConflictSharedUtils.h
--- a/Engine/Plugins/Developer/Concert/ConcertSync/ConcertSyncCore/Source/ConcertSyncCore/Public/Replication/AuthorityConflictSharedUtils.h
+++ b/Engine/Plugins/Developer/Concert/ConcertSync/ConcertSyncCore/Source/ConcertSyncCore/Public/Replication/AuthorityConflictSharedUtils.h
@@ -7,7 +7,6 @@
 #include "Misc/EBreakBehavior.h"
 #include "Templates/Function.h"
 
-
 struct FConcertPropertyChain;
 struct FConcertPropertySelection;
 struct FConcertReplication_ChangeAuthority_Request;
@@ -15,6 +14,8 @@ struct FConcertReplication_ChangeStream_Request;
 struct FConcertObjectReplicationMap;
 struct FConcertReplicatedObjectId;
 
+namespace UE::ConcertSyncCore::Replication { class IReplicationGroundTruth; }
+
 namespace UE::ConcertSyncCore::Replication::AuthorityConflictUtils
 {
 	enum class EAuthorityConflict : uint8
@@ -26,23 +27,6 @@ namespace UE::ConcertSyncCore::Replication::AuthorityConflictUtils
 	};
 	
 	using FProcessAuthorityConflict = TFunctionRef<EBreakBehavior(const FGuid& ClientId, const FGuid& StreamId, const FConcertPropertyChain& ConflictingProperty)>;
-
-	/** Defines a way to obtain streaming information about clients. */
-	class IReplicationGroundTruth
-	{
-	public:
-
-		/** Provides a way to extract all streams registered to a given client. */
-		virtual void ForEachStream(const FGuid& ClientEndpointId, TFunctionRef<EBreakBehavior(const FGuid& StreamId, const FConcertObjectReplicationMap& ReplicationMap)> Callback) const = 0;
-
-		/** Iterates through all clients have registered to send any data. */
-		virtual void ForEachClient(TFunctionRef<EBreakBehavior(const FGuid& ClientEndpointId)> Callback) const = 0;
-
-		/** @return Whether ClientId's stream StreamId has authority over ObjectPath. */
-		virtual bool HasAuthority(const FGuid& ClientId, const FGuid& StreamId, const FSoftObjectPath& ObjectPath) const = 0; 
-
-		virtual ~IReplicationGroundTruth() = default;
-	};
 	
 	/**
 	 * Enumerates all authority conflicts, if any, that would occur if Object were to be replicated with the given Properties.
diff --git a/Engine/Plugins/Developer/Concert/ConcertSync/ConcertSyncCore/Source/ConcertSyncCore/Public/Replication/Misc/IReplicationGroundTruth.h b/Engine/Plugins/Developer/Concert/ConcertSync/ConcertSyncCore/Source/ConcertSyncCore/Public/Replication/Misc/IReplicationGroundTruth.h
new file mode 100644
--- /dev/null
+++ b/Engine/Plugins/Developer/Concert/ConcertSync/ConcertSyncCore/Source/ConcertSyncCore/Public/Replication/Misc/IReplicationGroundTruth.h
@@ -0,0 +1,30 @@
+﻿// Copyright Epic Games, Inc. All Rights Reserved.
+
+#pragma once
+
+#include "HAL/Platform.h"
+#include "Misc/EBreakBehavior.h"
+#include "Templates/Function.h"
+
+struct FConcertObjectReplicationMap;
+struct FSoftObjectPath;
+
+namespace UE::ConcertSyncCore::Replication
+{
+	/** Defines a way to obtain streaming information about clients. */
+	class IReplicationGroundTruth
+	{
+	public:
+
+		/** Provides a way to extract all streams registered to a given client. */
+		virtual void ForEachStream(const FGuid& ClientEndpointId, TFunctionRef<EBreakBehavior(const FGuid& StreamId, const FConcertObjectReplicationMap& ReplicationMap)> Callback) const = 0;
+
+		/** Iterates through all clients have registered to send any data. */
+		virtual void ForEachClient(TFunctionRef<EBreakBehavior(const FGuid& ClientEndpointId)> Callback) const = 0;
+
+		/** @return Whether ClientId's stream StreamId has authority over ObjectPath. */
+		virtual bool HasAuthority(const FGuid& ClientId, const FGuid& StreamId, const FSoftObjectPath& ObjectPath) const = 0; 
+
+		virtual ~IReplicationGroundTruth() = default;
+	};
+}
diff --git a/Engine/Plugins/Developer/Concert/ConcertSync/ConcertSyncCore/Source/ConcertSyncCore/Public/Replication/Misc/StreamAndAuthorityPredictionUtils.h b/Engine/Plugins/Developer/Concert/ConcertSync/ConcertSyncCore/Source/ConcertSyncCore/Public/Replication/Misc/StreamAndAuthorityPredictionUtils.h
new file mode 100644
--- /dev/null
+++ b/Engine/Plugins/Developer/Concert/ConcertSync/ConcertSyncCore/Source/ConcertSyncCore/Public/Replication/Misc/StreamAndAuthorityPredictionUtils.h
@@ -0,0 +1,18 @@
+﻿// Copyright Epic Games, Inc. All Rights Reserved.
+
+#pragma once
+
+#include "ConcertMessageData.h"
+
+namespace UE::ConcertSyncCore::Replication
+{
+	/**
+	 * Decides whether First and Second should be considered to represent the same user across several Concert sessions.
+	 *
+	 * Every time a user joins a Concert session, a new endpoint ID is generated for that user and saved in the database.
+	 * Even though the endpoint ID is different, we can associate the same user across the IDs by using the DisplayName and DeviceName.
+	 *
+	 * @return Whether First and Second logically describe the same client (i.e. DisplayName and DeviceName of both clients are equal).
+	 */
+	CONCERTSYNCCORE_API bool AreLogicallySameClients(const FConcertClientInfo& First, const FConcertClientInfo& Second);
+}
diff --git a/Engine/Plugins/Developer/Concert/ConcertSync/ConcertSyncServer/Source/ConcertSyncServer/Private/Replication/AuthorityManager.cpp b/Engine/Plugins/Developer/Concert/ConcertSync/ConcertSyncServer/Source/ConcertSyncServer/Private/Replication/AuthorityManager.cpp
--- a/Engine/Plugins/Developer/Concert/ConcertSync/ConcertSyncServer/Source/ConcertSyncServer/Private/Replication/AuthorityManager.cpp
+++ b/Engine/Plugins/Developer/Concert/ConcertSync/ConcertSyncServer/Source/ConcertSyncServer/Private/Replication/AuthorityManager.cpp
@@ -8,13 +8,14 @@
 #include "Replication/Data/ObjectIds.h"
 #include "Replication/Data/ReplicationStream.h"
 #include "Replication/Messages/ChangeAuthority.h"
+#include "Replication/Misc/IReplicationGroundTruth.h"
 #include "Util/GroundTruthOverride.h"
 
 namespace UE::ConcertSyncServer::Replication
 {
 	namespace Private
 	{
-		class FServerGroundTruth : public ConcertSyncCore::Replication::AuthorityConflictUtils::IReplicationGroundTruth
+		class FServerGroundTruth : public ConcertSyncCore::Replication::IReplicationGroundTruth
 		{
 			const FAuthorityManager& Owner;
 			const IRegistrationEnumerator& Getters;
@@ -74,7 +75,7 @@ namespace UE::ConcertSyncServer::Replication
 		}
 
 		const FConcertObjectReplicationMap* FindClientReplicationMapById(
-			const ConcertSyncCore::Replication::AuthorityConflictUtils::IReplicationGroundTruth& Getters,
+			const ConcertSyncCore::Replication::IReplicationGroundTruth& Getters,
 			const FGuid& ClientId,
 			const FGuid& StreamId
 			)
diff --git a/Engine/Plugins/Developer/Concert/ConcertSync/ConcertSyncServer/Source/ConcertSyncServer/Private/Replication/ConcertServerReplicationManager_RestoreContent.cpp b/Engine/Plugins/Developer/Concert/ConcertSync/ConcertSyncServer/Source/ConcertSyncServer/Private/Replication/ConcertServerReplicationManager_RestoreContent.cpp
--- a/Engine/Plugins/Developer/Concert/ConcertSync/ConcertSyncServer/Source/ConcertSyncServer/Private/Replication/ConcertServerReplicationManager_RestoreContent.cpp
+++ b/Engine/Plugins/Developer/Concert/ConcertSync/ConcertSyncServer/Source/ConcertSyncServer/Private/Replication/ConcertServerReplicationManager_RestoreContent.cpp
@@ -1,6 +1,8 @@
 ﻿// Copyright Epic Games, Inc. All Rights Reserved.
 
 #include "ConcertLogGlobal.h"
+
+#include "ConcertMessageData.h"
 #include "ConcertServerReplicationManager.h"
 #include "ConcertServer/Private/ConcertServerSession.h"
 #include "Replication/IReplicationWorkspace.h"
@@ -9,17 +11,12 @@
 #include "Replication/Messages/ChangeStream.h"
 #include "Replication/Messages/ReplicationActivity.h"
 #include "Replication/Misc/ReplicationStreamUtils.h"
+#include "Replication/Misc/StreamAndAuthorityPredictionUtils.h"
 
 namespace UE::ConcertSyncServer::Replication
 {
 	namespace RestoreContentPrivate
 	{
-		/** @return Whether to consider the client infos equal for the purposes of restoring content */
-		static bool EqualsClientInfo(const FConcertClientInfo& Left, const FConcertClientInfo& Right)
-		{
-			return Left.DisplayName == Right.DisplayName && Left.DeviceName == Right.DeviceName;
-		}
-
 		/** Checks whether the request should fail according to EConcertReplicationRestoreContentFlags::ValidateUniqueClient. */
 		static bool HasNameConflict(
 			const FConcertSessionClientInfo& SenderInfo,
@@ -39,7 +36,7 @@ namespace UE::ConcertSyncServer::Replication
 
 				if (ensure(bFoundClient)
 					&& OtherInfo.ClientEndpointId != SenderInfo.ClientEndpointId
-					&& EqualsClientInfo(SenderInfo.ClientInfo, OtherInfo.ClientInfo))
+					&& ConcertSyncCore::Replication::AreLogicallySameClients(SenderInfo.ClientInfo, OtherInfo.ClientInfo))
 				{
 					return true;
 				}
diff --git a/Engine/Plugins/Developer/Concert/ConcertSync/ConcertSyncServer/Source/ConcertSyncServer/Private/Replication/Util/GroundTruthOverride.h b/Engine/Plugins/Developer/Concert/ConcertSync/ConcertSyncServer/Source/ConcertSyncServer/Private/Replication/Util/GroundTruthOverride.h
--- a/Engine/Plugins/Developer/Concert/ConcertSync/ConcertSyncServer/Source/ConcertSyncServer/Private/Replication/Util/GroundTruthOverride.h
+++ b/Engine/Plugins/Developer/Concert/ConcertSync/ConcertSyncServer/Source/ConcertSyncServer/Private/Replication/Util/GroundTruthOverride.h
@@ -3,6 +3,7 @@
 #pragma once
 
 #include "Replication/AuthorityConflictSharedUtils.h"
+#include "Replication/Misc/IReplicationGroundTruth.h"
 
 #include "Containers/Array.h"
 #include "Containers/Map.h"
@@ -27,7 +28,7 @@ namespace UE::ConcertSyncServer::Replication
 	 * This is useful if you want to validate i.e. that no authority conflicts happen if the injected clients were present. 
 	 */
 	class FGroundTruthOverride
-		: public ConcertSyncCore::Replication::AuthorityConflictUtils::IReplicationGroundTruth
+		: public ConcertSyncCore::Replication::IReplicationGroundTruth
 		, public FNoncopyable
 	{
 	public:
diff --git a/Engine/Plugins/Developer/Concert/ConcertSync/ConcertSyncTest/Source/ConcertSyncTest/Private/Replication/ReplicationAuthorityTests.cpp b/Engine/Plugins/Developer/Concert/ConcertSync/ConcertSyncTest/Source/ConcertSyncTest/Private/Replication/ReplicationAuthorityTests.cpp
--- a/Engine/Plugins/Developer/Concert/ConcertSync/ConcertSyncTest/Source/ConcertSyncTest/Private/Replication/ReplicationAuthorityTests.cpp
+++ b/Engine/Plugins/Developer/Concert/ConcertSync/ConcertSyncTest/Source/ConcertSyncTest/Private/Replication/ReplicationAuthorityTests.cpp
@@ -2,9 +2,11 @@
 
 #include "Util/SendReceiveObjectTestBase.h"
 
+#include "Replication/AuthorityConflictSharedUtils.h"
 #include "Replication/Formats/FullObjectFormat.h"
 #include "Replication/IConcertClientReplicationManager.h"
 #include "Replication/Messages/ObjectReplication.h"
+#include "Replication/Misc/IReplicationGroundTruth.h"
 #include "TestReflectionObject.h"
 #include "Util/ConcertClientReplicationBridgeMock.h"
 #include "Util/SendReceiveGenericStreamTestBase.h"
@@ -13,7 +15,6 @@
 #include "Misc/AutomationTest.h"
 #include "Misc/Paths.h"
 #include "HAL/FileManager.h"
-#include "Replication/AuthorityConflictSharedUtils.h"
 
 namespace UE::ConcertSyncTests::Replication::Authority
 {
@@ -372,7 +373,7 @@ namespace UE::ConcertSyncTests::Replication::Authority
 		constexpr FGuid RequestingClientId { 0, 0, 0, 0};
 		constexpr FGuid ExistingClientId { 1, 0, 0, 0};
 		
-		class FTestGroundTruth : public ConcertSyncCore::Replication::AuthorityConflictUtils::IReplicationGroundTruth
+		class FTestGroundTruth : public ConcertSyncCore::Replication::IReplicationGroundTruth
 		{
 		public:
 			
