Add a custom reader/writer to ConcertSyncDatabase.cpp to support FSoftObjectPath

diff --git a/Engine/Plugins/Developer/Concert/ConcertSync/ConcertSyncCore/Source/ConcertSyncCore/Private/ConcertSyncSessionDatabase.cpp b/Engine/Plugins/Developer/Concert/ConcertSync/ConcertSyncCore/Source/ConcertSyncCore/Private/ConcertSyncSessionDatabase.cpp
--- a/Engine/Plugins/Developer/Concert/ConcertSync/ConcertSyncCore/Source/ConcertSyncCore/Private/ConcertSyncSessionDatabase.cpp
+++ b/Engine/Plugins/Developer/Concert/ConcertSync/ConcertSyncCore/Source/ConcertSyncCore/Private/ConcertSyncSessionDatabase.cpp
@@ -72,6 +72,30 @@ FOnPackageSaved& UE::ConcertSyncCore::SyncDatabase::GetOnPackageSavedDelegate()
 	return SavedDelegate;
 }
 
+class FMemorySyncDbWriter : public FMemoryWriter
+{
+public:
+	using FMemoryWriter::FMemoryWriter;
+	using FMemoryWriter::operator<<;
+	virtual FArchive& operator<<(FSoftObjectPath& AssetPtr) override
+	{
+		AssetPtr.SerializePath(*this);
+		return *this;
+	}
+};
+
+class FMemorySyncDbReader : public FMemoryReader
+{
+public:
+	using FMemoryReader::FMemoryReader;
+	using FMemoryReader::operator<<;
+	virtual FArchive& operator<<(FSoftObjectPath& AssetPtr) override
+	{
+		AssetPtr.SerializePath(*this);
+		return *this;
+	}
+};
+
 namespace TransactionDataUtil
 {
 
@@ -94,7 +118,7 @@ FString GetDataFilename(const int64 InIndex)
 bool WriteTransactionData(const FStructOnScope& InTransaction, TArray<uint8>& OutSerializedTransactionData)
 {
 	SCOPED_CONCERT_TRACE(ConcertSyncSessionDatabase_WriteTransactionData);
-	FMemoryWriter Ar(OutSerializedTransactionData);
+	FMemorySyncDbWriter Ar(OutSerializedTransactionData);
 
 	const UScriptStruct* TransactionType = CastChecked<const UScriptStruct>(InTransaction.GetStruct());
 
@@ -110,7 +134,7 @@ bool WriteTransaction(const FStructOnScope& InTransaction, TArray<uint8>& OutSer
 	SCOPED_CONCERT_TRACE(ConcertSyncSessionDatabase_WriteTransaction);
 	check(InTransaction.IsValid());
 
-	FMemoryWriter Ar(OutSerializedTransactionData);
+	FMemorySyncDbWriter Ar(OutSerializedTransactionData);
 
 	// Serialize the data version
 	uint32 SerializedDataVersion = DataVersion;
@@ -153,7 +177,7 @@ bool WriteTransaction(const FStructOnScope& InTransaction, TArray<uint8>& OutSer
 bool ReadTransactionData(const TArray<uint8>& InSerializedTransactionData, FStructOnScope& OutTransaction)
 {
 	SCOPED_CONCERT_TRACE(ConcertSyncSessionDatabase_ReadTransactionData);
-	FMemoryReader Ar(InSerializedTransactionData);
+	FMemorySyncDbReader Ar(InSerializedTransactionData);
 
 	// Deserialize the transaction
 	UScriptStruct* TransactionType = nullptr;
@@ -186,7 +210,7 @@ bool ReadTransactionData(const TArray<uint8>& InSerializedTransactionData, FStru
 bool ReadTransaction(const TArray<uint8>& InSerializedTransactionData, FStructOnScope& OutTransaction)
 {
 	SCOPED_CONCERT_TRACE(ConcertSyncSessionDatabase_ReadTransaction);
-	FMemoryReader Ar(InSerializedTransactionData);
+	FMemorySyncDbReader Ar(InSerializedTransactionData);
 
 	// Test the footer is in place so we know we didn't crash mid-write
 	bool bParsedFooter = false;
@@ -427,7 +451,7 @@ WritePackageResult WritePackage(FConcertPackageDataStream& PackageDataStream, FA
 	SCOPED_CONCERT_TRACE(ConcertSyncSessionDatabase_WritePackage);
 
 	WritePackageResult OutResult;
-	FMemoryWriter MemWriter(OutResult.CacheData);
+	FMemorySyncDbWriter MemWriter(OutResult.CacheData);
 
 	EConcertCompressionDetails PackageDataFormat = UE::Concert::Compression::GetCompressionDetails(PackageDataStream.DataSize);
 	if (WillUseMemoryWriter(PackageDataStream.DataSize))
@@ -495,7 +519,7 @@ bool ExtractPackageData(FArchive& PackageBlobAr, const TFunctionRef<void(FConcer
 		}
 
 		// Let the caller stream the package data from a memory archive.
-		FMemoryReader DataAr(DecompressedData);
+		FMemorySyncDbReader DataAr(DecompressedData);
 		FConcertPackageDataStream PackageDataStream{ &DataAr, UncompressedPackageSize, &DecompressedData };
 		PackageDataStreamFallbackFn(PackageDataStream);
 	}
@@ -518,7 +542,7 @@ bool ExtractPackageData(FArchive& PackageBlobAr, const TFunctionRef<void(FConcer
 bool ExtractPackageData(const TArray<uint8>& InPackageBlob, const TFunctionRef<void(FConcertPackageDataStream&)>& PackageDataStreamFallbackFn)
 {
 	SCOPED_CONCERT_TRACE(ConcertSyncSessionDatabase_ExtreactPackageData);
-	FMemoryReader PackageBlobAr(InPackageBlob);
+	FMemorySyncDbReader PackageBlobAr(InPackageBlob);
 	return ExtractPackageData(PackageBlobAr, PackageDataStreamFallbackFn);
 }
 
@@ -2998,7 +3022,7 @@ bool FConcertSyncSessionDatabase::LoadPackage(const FString& InPackageBlobFilena
 	{
 		FConcertPackageAsyncDataStreamPtr Owned = *Found;
 		FConcertPackageDataStream Stream = Owned->PackageStream;
-		FMemoryReader PackageBlobAr(Owned->PackageData);
+		FMemorySyncDbReader PackageBlobAr(Owned->PackageData);
 		Stream.DataAr = &PackageBlobAr;
 		PackageDataStreamFn(Stream);
 		return true;
@@ -3110,7 +3134,7 @@ void FConcertSyncSessionDatabase::ScheduleAsyncWrite(const FName& InPackageName,
 		ensureMsgf(DstAr, TEXT("ConcertSyncSessionDatabase: Unable to open a file for writing. Ensure you have enough disk space and long paths are enabled."));
 
 		FConcertPackageDataStream& InPackageDataStream = SharedStream->PackageStream;
-		FMemoryReader PackageBlobAr(SharedStream->PackageData);
+		FMemorySyncDbReader PackageBlobAr(SharedStream->PackageData);
 		InPackageDataStream.DataAr = &PackageBlobAr;
 		SharedStream->Result = PackageDataUtil::WritePackage(InPackageDataStream, DstAr.Get());
 		return DstAr && !DstAr->IsError();
