Introduced algorithm to remap FSoftObjectPaths to different FSoftObjectPaths in FConcertObjectReplicationMap by actor label.

diff --git a/Engine/Plugins/Developer/Concert/ConcertSync/ConcertSyncCore/Source/ConcertSyncCore/Public/Replication/Data/ActorLabelRemapping.h b/Engine/Plugins/Developer/Concert/ConcertSync/ConcertSyncCore/Source/ConcertSyncCore/Public/Replication/Data/ActorLabelRemapping.h
new file mode 100644
--- /dev/null
+++ b/Engine/Plugins/Developer/Concert/ConcertSync/ConcertSyncCore/Source/ConcertSyncCore/Public/Replication/Data/ActorLabelRemapping.h
@@ -0,0 +1,73 @@
+﻿// Copyright Epic Games, Inc. All Rights Reserved.
+
+#pragma once
+
+#include "Containers/Map.h"
+#include "Containers/UnrealString.h"
+#include "UObject/SoftObjectPath.h"
+#include "UObject/SoftObjectPtr.h"
+
+#include "ActorLabelRemapping.generated.h"
+
+/** Metadata about an actor in a FConcertObjectReplicationMap. */
+USTRUCT()
+struct FConcertReplicationRemappingData_Actor
+{
+	GENERATED_BODY()
+	
+	/**
+	 * Label of the actor.
+	 * 
+	 * This is typically AActor::ActorLabel.
+	 * Practically, the API is generic enough to allow using some other actor identifying string at runtime instead if you really wanted to.
+	 */
+	UPROPERTY()
+	FString Label;
+
+	/**
+	 * The class the actor had when the remapping data was generated.
+	 * 
+	 * Note that the actor's itself may not even have been in the FConcertObjectReplicationMap when the data was generated.
+	 * During generation, when processing a component (or other subobject), the owning actor instance can be looked up by path and its class is recorded here.
+	 * If the actor does not resolve, it is not recorded.
+	 *
+	 * Example:
+	 *  - Suppose the replication contains a single component: "/Game/World.World:PersistentLevel.Actor.Component"
+	 *	- If during data generation, the class of "/Game/World.World:PersistentLevel.Actor" 
+	 *		- can be obtained, the corresponding FConcertReplicationRemappingData::ActorData contains the actor path and its class.
+	 *		- can not be obtained, then FConcertReplicationRemappingData::Actor.IsEmpty()
+	 */
+	UPROPERTY()
+	FSoftClassPath Class;
+};
+
+/**
+ * Supplementary data for a FConcertObjectReplicationMap.
+ * This data can be used to remap FSoftObjectPaths in FConcertObjectReplicationMap to different FConcertObjectReplicationMap based on AActor::ActorLabel.
+ * The primary use case is replication presets being translated multiple sessions.
+ * 
+ * Example:
+ * 1. In session 1, suppose a FConcertObjectReplicationMap is created like this:
+ *		- { "/Game/World.World:PersistentLevel.ActorOld", { "ActorTags" } }
+ *		- { "/Game/World.World:PersistentLevel.ActorOld.StaticMeshComponentOld", { "RelativeLocation", "RelativeLocation.X" } }
+ *		- Actor_0 has the actor label "MyStaticMeshActor"
+ * 2. As part of creating the FConcertObjectReplicationMap, you also create a FConcertReplicationRemappingData (@see ActorLabelRemappingCore.h)
+ * 3. In session 2, suppose there is world:
+ *		- /Game/OtherWorld.OtherWorld:PersistentLevel.ActorNew
+ *		- /Game/OtherWorld.OtherWorld.PersistentLevel.ActorNew.StaticMeshComponentNew
+ *		- OtherActor_0 has the actor label "MyStaticMeshActor"
+ * In the above scenario, you can now remap the FConcertObjectReplicationMap from step 1 using the FConcertReplicationRemappingData from step 2.
+ * Effectively, "ActorNew" would get the properties of "ActorOld", and "StaticMeshComponentNew" the properties of "StaticMeshComponentOld", because
+ * the labels of "ActorNew" and "ActorOld" are both "MyStaticMeshActor".
+ *
+ * @see ActorLabelRemappingUtils.h
+ */
+USTRUCT()
+struct FConcertReplicationRemappingData
+{
+	GENERATED_BODY()
+	
+	/** Data about actors. */
+	UPROPERTY()
+	TMap<FSoftObjectPath, FConcertReplicationRemappingData_Actor> ActorData;
+};
diff --git a/Engine/Plugins/Developer/Concert/ConcertSync/ConcertSyncCore/Source/ConcertSyncCore/Public/Replication/Misc/ActorLabelRemappingCore.h b/Engine/Plugins/Developer/Concert/ConcertSync/ConcertSyncCore/Source/ConcertSyncCore/Public/Replication/Misc/ActorLabelRemappingCore.h
new file mode 100644
--- /dev/null
+++ b/Engine/Plugins/Developer/Concert/ConcertSync/ConcertSyncCore/Source/ConcertSyncCore/Public/Replication/Misc/ActorLabelRemappingCore.h
@@ -0,0 +1,149 @@
+﻿// Copyright Epic Games, Inc. All Rights Reserved.
+
+#pragma once
+
+#include "ActorLabelRemappingInput.h"
+#include "Misc/ObjectUtils.h"
+#include "Replication/Data/ActorLabelRemapping.h"
+#include "Replication/Data/ObjectReplicationMap.h"
+#include "Private/RemapAlgorithm.h"
+
+#include "Containers/UnrealString.h"
+#include "HAL/Platform.h"
+#include "Misc/Optional.h"
+#include "Templates/FunctionFwd.h"
+#include "UObject/SoftObjectPath.h"
+#include "UObject/SoftObjectPtr.h"
+
+enum class EBreakBehavior : uint8;
+
+namespace UE::ConcertSyncCore
+{
+	/**
+	 * Creates the remapping data for a replication map.
+	 * @param Origin The replication map to generate the remapping data for
+	 * @param GetLabelFunc Gets the label of an object
+	 * @param GetClassFunc Gets the class of an object
+	 * @return The remapping data you can use to remap Origin to.
+	 */
+	template<CGetObjectLabelCallable TGetObjectLabel, CGetObjectClassCallable TGetObjectClass>
+	void GenerateRemappingData(
+		const FConcertObjectReplicationMap& Origin,
+		TGetObjectLabel&& GetLabelFunc,
+		TGetObjectClass&& GetClassFunc,
+		FConcertReplicationRemappingData& Result
+		);
+	/** Util version that creates a new replication mapping. */
+	template<CGetObjectLabelCallable TGetObjectLabel, CGetObjectClassCallable TGetObjectClass>
+	FConcertReplicationRemappingData GenerateRemappingData(
+		const FConcertObjectReplicationMap& Origin,
+		TGetObjectLabel&& GetLabelFunc,
+		TGetObjectClass&& GetClassFunc
+		)
+	{
+		FConcertReplicationRemappingData Result;
+		GenerateRemappingData(Origin, GetLabelFunc, GetClassFunc, Result);
+		return Result;
+	}
+
+	/**
+	 * Attempts to remap all objects in Origin according to the labels saved in RemappingData.
+	 * Calls ProcessRemapping for each object so remapped.
+	 *
+	 * @see FConcertReplicationRemappingData for an example.
+	 *
+	 * @param Origin The replicationmap to remap
+	 * @param RemappingData The data based on which to remap
+	 * @param IsRemappingCompatibleFunc Decides whether this remapping is compatible
+	 * @param ForEachObjectWithLabelFunc Lists all objects with a particular label.
+	 * @param GetLabelFunc Gets the label of an object 
+	 * @param ProcessRemapping Called for every remapped object
+	 */
+	template
+	<
+		CIsRemappingCompatibleCallable TIsRemappingCompatible,
+		CForEachObjectWithLabelCallable TForEachObjectWithLabel,
+		CGetObjectLabelCallable TGetObjectLabel,
+		CProcessRemappingCallbable TProcessRemapping
+	>
+	void RemapReplicationMap(
+		const FConcertObjectReplicationMap& Origin,
+		const FConcertReplicationRemappingData& RemappingData,
+		TIsRemappingCompatible&& IsRemappingCompatibleFunc,
+		TForEachObjectWithLabel&& ForEachObjectWithLabelFunc,
+		TGetObjectLabel&& GetLabelFunc,
+		TProcessRemapping&& ProcessRemapping
+		)
+	{
+		Private::TRemapAlgorithm(Origin, RemappingData, IsRemappingCompatibleFunc, ForEachObjectWithLabelFunc, GetLabelFunc)
+			.Run([&ProcessRemapping](const FSoftObjectPath& Original, const FSoftObjectPath& Target)
+			{
+				ProcessRemapping(Original, Target);
+			});
+	}
+	/** Alternate version that directly writes into OutTargetMap. */
+	template<CIsRemappingCompatibleCallable TIsRemappingCompatible, CForEachObjectWithLabelCallable TForEachObjectWithLabel, CGetObjectLabelCallable TGetObjectLabel>
+	void RemapReplicationMap(
+		const FConcertObjectReplicationMap& Origin,
+		const FConcertReplicationRemappingData& RemappingData,
+		TIsRemappingCompatible&& IsRemappingCompatibleFunc,
+		TForEachObjectWithLabel&& ForEachObjectWithLabelFunc,
+		TGetObjectLabel&& GetLabelFunc,
+		FConcertObjectReplicationMap& OutTargetMap
+		)
+	{
+		Private::TRemapAlgorithm(Origin, RemappingData, IsRemappingCompatibleFunc, ForEachObjectWithLabelFunc, GetLabelFunc)
+			.Run([&Origin, &OutTargetMap](const FSoftObjectPath& Original, const FSoftObjectPath& Target)
+			{
+				OutTargetMap.ReplicatedObjects.Add(Target, Origin.ReplicatedObjects[Original]);
+			});
+	}
+	/** Util version that creates a new replication mapping. */
+	template<CIsRemappingCompatibleCallable TIsRemappingCompatible, CForEachObjectWithLabelCallable TForEachObjectWithLabel, CGetObjectLabelCallable TGetObjectLabel>
+	FConcertObjectReplicationMap RemapReplicationMap(
+		const FConcertObjectReplicationMap& Origin,
+		const FConcertReplicationRemappingData& RemappingData,
+		TIsRemappingCompatible&& IsRemappingCompatibleFunc,
+		TForEachObjectWithLabel&& ForEachObjectWithLabelFunc,
+		TGetObjectLabel&& GetLabelFunc
+		)
+	{
+		FConcertObjectReplicationMap Result;
+		RemapReplicationMap(Origin, RemappingData, IsRemappingCompatibleFunc, ForEachObjectWithLabelFunc, GetLabelFunc, Result);
+		return Result;
+	}
+}
+
+namespace UE::ConcertSyncCore
+{
+	template <CGetObjectLabelCallable TGetObjectLabel, CGetObjectClassCallable TGetObjectClass>
+	void GenerateRemappingData(
+		const FConcertObjectReplicationMap& Origin,
+		TGetObjectLabel&& GetLabelFunc,
+		TGetObjectClass&& GetClassFunc,
+		FConcertReplicationRemappingData& Result
+		)
+	{
+		for (const TPair<FSoftObjectPath, FConcertReplicatedObjectInfo>& Pair : Origin.ReplicatedObjects)
+		{
+			const FSoftObjectPath& ObjectPath = Pair.Key;
+			const TOptional<FSoftObjectPath> OwningActorPath = GetActorPathIn(ObjectPath);
+			if (!OwningActorPath || Result.ActorData.Contains(*OwningActorPath))
+			{
+				continue;
+			}
+
+			// GetLabelFunc and GetClassFunc *could* be combined into a single func to reduce ObjectPtr.Resolve count from 2 to 1 but the overhead
+			// should be minimal for the 2nd resolve because ObjectPtr will cache the value; we're making it easier for the API user by having 2 funcs.
+			const FSoftObjectPtr OwningActorPtr{ *OwningActorPath };
+			const TOptional<FString> Label = GetLabelFunc(OwningActorPtr);
+			const FSoftClassPath Class = GetClassFunc(OwningActorPtr);
+			if (Label && !Class.IsNull())
+			{
+				FConcertReplicationRemappingData_Actor& ActorData = Result.ActorData.Add(*OwningActorPath);
+				ActorData.Label = *Label;
+				ActorData.Class = Class;
+			}
+		}
+	}
+}
diff --git a/Engine/Plugins/Developer/Concert/ConcertSync/ConcertSyncCore/Source/ConcertSyncCore/Public/Replication/Misc/ActorLabelRemappingEditor.h b/Engine/Plugins/Developer/Concert/ConcertSync/ConcertSyncCore/Source/ConcertSyncCore/Public/Replication/Misc/ActorLabelRemappingEditor.h
new file mode 100644
--- /dev/null
+++ b/Engine/Plugins/Developer/Concert/ConcertSync/ConcertSyncCore/Source/ConcertSyncCore/Public/Replication/Misc/ActorLabelRemappingEditor.h
@@ -0,0 +1,102 @@
+﻿// Copyright Epic Games, Inc. All Rights Reserved.
+
+#pragma once
+
+#if WITH_EDITOR
+
+#include "ActorLabelRemappingInput.h"
+#include "ActorLabelRemappingCore.h"
+#include "Private/EditorRemappingUtils.h"
+
+class UWorld;
+
+namespace UE::ConcertSyncCore
+{
+	// The functions in ActorLabelRemappingCore.h are generic and operate only on FSoftObjectPaths.
+	// The following functions are only useful in builds that have AActor instances.
+	
+	/** This version of GenerateRemappingData uses AActor::ActorLabel. */
+	void GenerateRemappingData(
+		const FConcertObjectReplicationMap& Origin,
+		FConcertReplicationRemappingData& Result
+		);
+	/** Util version that creates a new replication mapping. */
+	FConcertReplicationRemappingData GenerateRemappingData(const FConcertObjectReplicationMap& Origin);
+
+	/** This version of RemapReplicationMap uses AActor::ActorLabel and tries to remap to actors in TargetWorld. */
+	template<CProcessRemappingCallbable TProcessRemapping>
+	void RemapReplicationMap(
+		const FConcertObjectReplicationMap& Origin,
+		const FConcertReplicationRemappingData& RemappingData,
+		const UWorld& TargetWorld,
+		TProcessRemapping&& ProcessRemapping
+		);
+	/** Alternate version that directly writes into OutTargetMap. */
+	void RemapReplicationMap(
+		const FConcertObjectReplicationMap& Origin,
+		const FConcertReplicationRemappingData& RemappingData,
+		const UWorld& TargetWorld,
+		FConcertObjectReplicationMap& OutTargetMap
+		);
+	/** Util version that creates a new replication mapping. */
+	FConcertObjectReplicationMap RemapReplicationMap(
+		const FConcertObjectReplicationMap& Origin,
+		const FConcertReplicationRemappingData& RemappingData,
+		const UWorld& TargetWorld
+		);
+}
+
+namespace UE::ConcertSyncCore // Inline defs
+{
+	FORCEINLINE void GenerateRemappingData(
+		const FConcertObjectReplicationMap& Origin,
+		FConcertReplicationRemappingData& Result
+		)
+	{
+		GenerateRemappingData(
+			Origin,
+			[](const FSoftObjectPtr& Object){ return Private::GetActorLabel(Object); },
+			[](const FSoftObjectPtr& Object) { return Private::GetClassPath(Object); },
+			Result
+			);
+	}
+	
+	FORCEINLINE FConcertReplicationRemappingData GenerateRemappingData(const FConcertObjectReplicationMap& Origin)
+	{
+		FConcertReplicationRemappingData Result;
+		GenerateRemappingData(Origin, Result);
+		return Result;
+	}
+	
+	template <CProcessRemappingCallbable TProcessRemapping>
+	void RemapReplicationMap(
+		const FConcertObjectReplicationMap& Origin,
+		const FConcertReplicationRemappingData& RemappingData,
+		const UWorld& TargetWorld,
+		TProcessRemapping&& ProcessRemapping)
+	{
+		Private::GenericRemapReplicationMap(Origin, RemappingData, TargetWorld, ProcessRemapping);
+	}
+	
+	FORCEINLINE void RemapReplicationMap(
+		const FConcertObjectReplicationMap& Origin,
+		const FConcertReplicationRemappingData& RemappingData,
+		const UWorld& TargetWorld,
+		FConcertObjectReplicationMap& OutTargetMap
+		)
+	{
+		Private::GenericRemapReplicationMap(Origin, RemappingData, TargetWorld, OutTargetMap);
+	}
+	
+	FORCEINLINE FConcertObjectReplicationMap RemapReplicationMap(
+		const FConcertObjectReplicationMap& Origin,
+		const FConcertReplicationRemappingData& RemappingData,
+		const UWorld& TargetWorld
+		)
+	{
+		FConcertObjectReplicationMap Result;
+		RemapReplicationMap(Origin, RemappingData, TargetWorld, Result);
+		return Result;
+	}
+}
+#endif
\ No newline at end of file
diff --git a/Engine/Plugins/Developer/Concert/ConcertSync/ConcertSyncCore/Source/ConcertSyncCore/Public/Replication/Misc/ActorLabelRemappingInput.h b/Engine/Plugins/Developer/Concert/ConcertSync/ConcertSyncCore/Source/ConcertSyncCore/Public/Replication/Misc/ActorLabelRemappingInput.h
new file mode 100644
--- /dev/null
+++ b/Engine/Plugins/Developer/Concert/ConcertSync/ConcertSyncCore/Source/ConcertSyncCore/Public/Replication/Misc/ActorLabelRemappingInput.h
@@ -0,0 +1,70 @@
+﻿// Copyright Epic Games, Inc. All Rights Reserved.
+
+#pragma once
+
+#include "HAL/Platform.h"
+#include "Templates/FunctionFwd.h"
+
+#include <type_traits>
+
+class FString;
+
+enum class EBreakBehavior : uint8;
+
+struct FSoftClassPath;
+struct FSoftObjectPath;
+struct FSoftObjectPtr;
+template<typename OptionalType> struct TOptional;
+
+namespace UE::ConcertSyncCore
+{
+	/**
+	 * Lets you decide whether OriginObject can be remapped to PossibleTarget.
+	 *
+	 * This should check the following:
+	 * - does PossibleTarget exist? PossibleTarget is the result of a simple string replacement in the object path.
+	 * - is PossibleTarget's class compatible with OriginClass? You can decide what compatible means - you may want exact classes or subclasses.
+	 *
+	 * @param OriginObject The original object that is being evaluated as remapping candidate.
+	 *	Important: if OriginObject points to an actor, then the original replication map is not guaranted to contain OriginObject.
+	 *	This happens e.g. if the original map contained "/Game/World.World:PersistentLevel.Actor.Component" but not its owning actor.
+	 *	In that case, this callback is being asked, whether the subobjects ("Component" in this case) can be remapped to actor OriginObject points at.
+	 * @param OriginClass The class the OriginObject is expected to be compatible with. 
+	 * @param PossibleTargetActor The owning actor of PossibleTarget. If PossibleTarget is an actor, then PossibleTarget == PossibleTargetActor.
+	 *	This arg may be useful because points to the actual UObject instance (if in engine build).
+	 * @param PossibleTarget The object that OriginObject will be remapped to if this function returns true.
+	 */
+	template<typename TCallback>
+	concept CIsRemappingCompatibleCallable = std::is_invocable_r_v<bool, TCallback,
+		const FSoftObjectPath& /*OriginObject*/,
+		const FSoftClassPath& /*OriginClass*/,
+		const FSoftObjectPtr& /*PossibleTargetActor*/,
+		const FSoftObjectPath& /*PossibleTarget*/
+	>;
+
+	/** For each object that has the given Label invokes Consumer until either Consumer returns EBreakBehavior::Break or there are no more */
+	template<typename TCallback>
+	concept CForEachObjectWithLabelCallable = std::is_invocable_r_v<void, TCallback,
+		const FString& /*Label*/,
+		TFunctionRef<EBreakBehavior(const FSoftObjectPtr& Actor)> /*Consumer*/
+		>;
+	
+	/**
+	 * Gets the label of an actor.
+	 * 
+	 * Return unset if any:
+	 * - Object not available
+	 * - you don't want the label in the result.
+	 */
+	template<typename TCallback>
+	concept CGetObjectLabelCallable = std::is_invocable_r_v<TOptional<FString>, TCallback, const FSoftObjectPtr& /*Object*/>;
+	/**
+	 * Gets an object's class.
+	 * Return an empty class path if Object cannot be resolved.
+	 */
+	template<typename TCallback>
+	concept CGetObjectClassCallable = std::is_invocable_r_v<FSoftClassPath, TCallback, const FSoftObjectPtr& /*Object*/>;
+	/** Process a remapped object */
+	template<typename TCallback>
+	concept CProcessRemappingCallbable = std::is_invocable_r_v<void, TCallback, const FSoftObjectPath& /*OriginalObject*/, const FSoftObjectPath& /*RemappedObject*/> ;
+}
\ No newline at end of file
diff --git a/Engine/Plugins/Developer/Concert/ConcertSync/ConcertSyncCore/Source/ConcertSyncCore/Public/Replication/Misc/Private/EditorRemappingUtils.h b/Engine/Plugins/Developer/Concert/ConcertSync/ConcertSyncCore/Source/ConcertSyncCore/Public/Replication/Misc/Private/EditorRemappingUtils.h
new file mode 100644
--- /dev/null
+++ b/Engine/Plugins/Developer/Concert/ConcertSync/ConcertSyncCore/Source/ConcertSyncCore/Public/Replication/Misc/Private/EditorRemappingUtils.h
@@ -0,0 +1,110 @@
+﻿// Copyright Epic Games, Inc. All Rights Reserved.
+
+#pragma once
+
+#if WITH_EDITOR
+
+#include "Replication/Misc/ActorLabelRemappingCore.h"
+
+#include "Containers/Array.h"
+#include "Containers/Map.h"
+#include "Containers/UnrealString.h"
+#include "EngineUtils.h"
+#include "GameFramework/Actor.h"
+#include "UObject/Object.h"
+#include "UObject/SoftObjectPath.h"
+#include "UObject/SoftObjectPtr.h"
+
+namespace UE::ConcertSyncCore::Private
+{
+	constexpr bool bCreateActorLabelDuringRead = false;
+		
+	inline TOptional<FString> GetActorLabel(const FSoftObjectPtr& Object)
+	{
+		const UObject* ResolvedObject = Object.Get();
+		const AActor* Actor = Cast<AActor>(ResolvedObject);
+		const FString Label = Actor ? Actor->GetActorLabel(bCreateActorLabelDuringRead) : FString{};
+		return Label.IsEmpty() ? TOptional<FString>{} : Label;
+	}
+
+	inline FSoftClassPath GetClassPath(const FSoftObjectPtr& Object)
+	{
+		const UObject* ResolvedObject = Object.Get();
+		return ResolvedObject ? ResolvedObject->GetClass() : FSoftClassPath{};
+	}
+	
+	inline TMap<FString, TArray<FSoftObjectPtr>> CacheByActorLabel(const UWorld& World)
+	{
+		TMap<FString, TArray<FSoftObjectPtr>> LabelsToInfo;
+		for (TActorIterator<AActor> ActorIt(&World); ActorIt; ++ActorIt)
+		{
+			const FString Label = ActorIt->GetActorLabel(bCreateActorLabelDuringRead);
+			if (Label.IsEmpty())
+			{
+				continue;;
+			}
+			
+			LabelsToInfo.FindOrAdd(Label).Emplace(*ActorIt);
+		}
+		return LabelsToInfo;
+	}
+
+	/** This functions sole purpose is to avoid code duplication for the overloads of GenerateRemappingData. */
+	template<typename TFinalArg>
+	static void GenericRemapReplicationMap(
+		const FConcertObjectReplicationMap& Origin,
+		const FConcertReplicationRemappingData& RemappingData,
+		const UWorld& TargetWorld,
+		TFinalArg&& Argument
+	)
+	{
+		const auto IsRemappingCompatibleFunc = [](
+			const FSoftObjectPath&,
+			const FSoftClassPath& OriginClass,
+			const FSoftObjectPtr&,
+			const FSoftObjectPath& PossibleTarget
+			)
+		{
+			UObject* Object = PossibleTarget.ResolveObject();
+			if (!Object)
+			{
+				return false;
+			}
+			
+			const bool bAreClassesSame = FSoftClassPath(Object->GetClass()) == OriginClass;
+			return bAreClassesSame;
+		};
+
+		using namespace Private;
+		const TMap<FString, TArray<FSoftObjectPtr>> LabelCache = CacheByActorLabel(TargetWorld);
+		const auto ForEachObjectWithLabel = [&LabelCache](
+			const FString& Label,
+			TFunctionRef<EBreakBehavior(const FSoftObjectPtr& Actor)> Consumer
+			)
+		{
+			const TArray<FSoftObjectPtr>* ActorsWithlabel = LabelCache.Find(Label);
+			if (!ActorsWithlabel)
+			{
+				return;
+			}
+
+			for (const FSoftObjectPtr& Actor : *ActorsWithlabel)
+			{
+				if (Consumer(Actor) == EBreakBehavior::Break)
+				{
+					break;
+				}
+			}
+		};
+	
+		RemapReplicationMap(
+			Origin,
+			RemappingData,
+			IsRemappingCompatibleFunc,
+			ForEachObjectWithLabel,
+			[](const FSoftObjectPtr& Object){ return GetActorLabel(Object); },
+			Argument
+			);
+	}
+}
+#endif
\ No newline at end of file
diff --git a/Engine/Plugins/Developer/Concert/ConcertSync/ConcertSyncCore/Source/ConcertSyncCore/Public/Replication/Misc/Private/RemapAlgorithm.h b/Engine/Plugins/Developer/Concert/ConcertSync/ConcertSyncCore/Source/ConcertSyncCore/Public/Replication/Misc/Private/RemapAlgorithm.h
new file mode 100644
--- /dev/null
+++ b/Engine/Plugins/Developer/Concert/ConcertSync/ConcertSyncCore/Source/ConcertSyncCore/Public/Replication/Misc/Private/RemapAlgorithm.h
@@ -0,0 +1,371 @@
+﻿// Copyright Epic Games, Inc. All Rights Reserved.
+
+#pragma once
+
+#include "ConcertLogGlobal.h"
+#include "Misc/ObjectPathHierarchy.h"
+#include "Misc/ObjectUtils.h"
+#include "Replication/Data/ActorLabelRemapping.h"
+#include "Replication/Data/ObjectReplicationMap.h"
+#include "Replication/Misc/ActorLabelRemappingInput.h"
+
+#include <type_traits>
+
+namespace UE::ConcertSyncCore::Private
+{
+	/**
+	 * This algorithm remaps FSoftObjectPaths from an origin FConcertObjectReplicationMap to target FSoftObjectPaths based on the owning actors
+	 * sharing actor labels.
+	 * @see FConcertReplicationRemappingData for example.
+	 */
+	template<CIsRemappingCompatibleCallable TIsRemappingCompatible, CForEachObjectWithLabelCallable TForEachObjectWithLabel, CGetObjectLabelCallable TGetObjectLabel>
+	class TRemapAlgorithm
+	{
+	public:
+
+		TRemapAlgorithm(
+			const FConcertObjectReplicationMap& Origin,
+			const FConcertReplicationRemappingData& RemappingData,
+			const TIsRemappingCompatible& ForEachObjectWithOuterFunc,
+			const TForEachObjectWithLabel& ForEachObjectWithLabelFunc,
+			const TGetObjectLabel& GetLabelFunc
+			)
+			: Origin(Origin)
+			, RemappingData(RemappingData)
+			, IsRemappingCompatibleFunc(ForEachObjectWithOuterFunc)
+			, ForEachObjectWithLabelFunc(ForEachObjectWithLabelFunc)
+			, GetLabelFunc(GetLabelFunc)
+		{}
+		
+		template<CProcessRemappingCallbable TProcessRemappingLambda>
+		void Run(TProcessRemappingLambda&& ProcessRemapping)
+		{
+			// 1. Build object hierarchy of Origin & find all actors that we are supposed to remap (those that have a label)
+			BuildOriginObjectHierarchyAndActorsNeedingRemapping();
+			
+			// 2. For each re-mappable actor, build the list of target candidates using the hierarchy to check the requirements.
+			ComputeRemappingCandidates();
+			
+			// 3. In order of least-remaining target candidates, remap actors; uses the hierarchy to remap actor subobjects.
+			RemapActors(ProcessRemapping);
+		}
+
+	private:
+	
+		// Input data to the algorithm
+		const FConcertObjectReplicationMap& Origin;
+		const FConcertReplicationRemappingData& RemappingData;
+		const TIsRemappingCompatible& IsRemappingCompatibleFunc;
+		const TForEachObjectWithLabel& ForEachObjectWithLabelFunc;
+		const TGetObjectLabel& GetLabelFunc;
+		
+		/****************************** 1. Discover actors with labels ******************************/
+		
+		struct FActorAndClassPair
+		{
+			/** The path of an actor that is directly or indirectly referenced by the replication map. */
+			const FSoftObjectPath OriginActor;
+			/** The class that OriginActor is expected to have, as recorded by FConcertReplicationRemappingData_Actor. */
+			const FSoftClassPath& Class;
+
+			FActorAndClassPair(FSoftObjectPath OriginActor, const FSoftClassPath& Class)
+				: OriginActor(MoveTemp(OriginActor))
+				, Class(Class)
+			{}
+		};
+		
+		/** Hierarchy that is used by ComputeTargetCandidates to efficiently determine objects that need to be reassigned. */
+		FObjectPathHierarchy ObjectPathHierarchy;
+		/** Maps actor labels to object path from the original map that had that label. */
+		TMap<FString, TArray<FActorAndClassPair>> LabelsToPendingActors;
+		int32 ExpectedNumberRemappedActors = 0;
+
+		/**
+		 * Processes the input replication map:
+		 * - Builds ObjectPathHierarchy based on what the original replication map contained - this is used in later for analyzing the hierarchy.
+		 * - Groups all actors that need processing by label in LabelsToPendingActors.
+		 */
+		void BuildOriginObjectHierarchyAndActorsNeedingRemapping()
+		{
+			for (const TPair<FSoftObjectPath, FConcertReplicatedObjectInfo>& Pair : Origin.ReplicatedObjects)
+			{
+				const FSoftObjectPath& ObjectPath = Pair.Key;
+				ObjectPathHierarchy.AddObject(ObjectPath);
+
+				// It could be that ObjectPath is a component and the owning actor is not replicated / was not looped, yet.
+				// Handle that case here.
+				const TOptional<FSoftObjectPath> OwningActorPath = GetActorPathIn(ObjectPath);
+				const FConcertReplicationRemappingData_Actor* ActorData = OwningActorPath ? RemappingData.ActorData.Find(*OwningActorPath) : nullptr;
+				if (!ActorData)
+				{
+					continue;
+				}
+				
+				TArray<FActorAndClassPair>& Actors = LabelsToPendingActors.FindOrAdd(ActorData->Label);
+				const bool bAlreadyContainsActor = Actors.ContainsByPredicate([&OwningActorPath](const FActorAndClassPair& Pair)
+				{
+					return Pair.OriginActor == *OwningActorPath;
+				});
+				// Since we do GetActorPathIn for every subobject above, we may already have added the actor
+				if (!bAlreadyContainsActor)
+				{
+					++ExpectedNumberRemappedActors;
+					Actors.Emplace(*OwningActorPath, ActorData->Class);
+				}
+			}
+		}
+		
+		/****************************** 2. Find candidates that actors can be remapped to ******************************/
+		
+		/**
+		 * The max number of actors in a level we expect to have share the same label.
+		 * We expect a typical level to have at most 4 actors with the same label... if we're wrong we're punished with dynamic memory allocs.
+		 */
+		enum { ExpectedNumLabelCollisions = 4 };
+		
+		template<typename T>
+		using TSmallInlineArray = TArray<T, TInlineAllocator<ExpectedNumLabelCollisions>>;
+		
+		struct FActorWithSolutions
+		{
+			/**
+			 * An actor directly or indirectly referenced by the Origin replication map.
+			 * This is a reference to FActorAndClassPair::OriginActor.
+			 */
+			const FSoftObjectPath& OriginActor;
+			/** Actors that OriginActor can be replaced with. It had been validated that the hierarchy is compatible. */
+			TSmallInlineArray<FSoftObjectPath> PossibleSolutions;
+
+			explicit FActorWithSolutions(const FSoftObjectPath& OriginActor)
+				: OriginActor(OriginActor)
+			{}
+		};
+		/** Holds original paths that can be remapped. This is a TArray instead of TMap because RemapActors will iterate it often. */
+		TArray<FActorWithSolutions> ObjectsWithSolutions;
+
+		/**
+		 * Goes through all actors of which the hierarchy needs remapping.
+		 * Invokes ForEachObjectWithLabelFunc to get candidate actors with the same label.
+		 * Proceeds to check whether the candidates' subobject hierarchy is compatible and if so, adds the candiate to ObjectsWithSolutions.
+		 */
+		void ComputeRemappingCandidates()
+		{
+			if (ExpectedNumberRemappedActors == 0)
+			{
+				return;
+			}
+			
+			ObjectsWithSolutions.Reserve(ExpectedNumberRemappedActors);
+			for (TPair<FString, TArray<FActorAndClassPair>>& Pair : LabelsToPendingActors)
+			{
+				// If we find no potential solutions, this will keep some entries with PossibleSolutions.Num() == 0. RemapActors handles that.
+				TArray<FActorAndClassPair>& PendingActors = Pair.Value;
+				const int32 StartIndex = ObjectsWithSolutions.Num();
+				Algo::Transform(PendingActors, ObjectsWithSolutions, [](const FActorAndClassPair& Pending)
+				{
+					return FActorWithSolutions{ Pending.OriginActor };
+				});
+				const auto GetSolutionData = [this, StartIndex](int32 ActorIndex) -> FActorWithSolutions&
+				{
+					const int32 SolutionIdx = StartIndex + ActorIndex;
+					return ObjectsWithSolutions[SolutionIdx];
+				};
+				
+				ForEachObjectWithLabelFunc(Pair.Key, [this, &PendingActors, &GetSolutionData](const FSoftObjectPtr& TargetCandidate)
+				{
+					FindSolutionsIn(PendingActors, TargetCandidate, GetSolutionData);
+					// In theory, no further processing is needed once every remapped actor has PendingActors.Num() solutions in PossibleSolutions.
+					// In practice, we don't expect that many actors with the label in the level to warrant the implementation.
+					return EBreakBehavior::Continue;
+				});
+			}
+		}
+
+		/** Analyzes for every provided actor whether PossibleTargetActor is a suitable substitute. */
+		template<typename TGetActorSolutionData>
+		requires std::is_invocable_r_v<FActorWithSolutions&, TGetActorSolutionData, int32>
+		void FindSolutionsIn(
+			const TArray<FActorAndClassPair>& PendingActors,
+			const FSoftObjectPtr& TargetCandidate,
+			TGetActorSolutionData&& GetSolutionData
+			)
+		{
+			for (int32 ActorIdx = 0; ActorIdx < PendingActors.Num(); ++ActorIdx)
+			{
+				const FActorAndClassPair& PendingActor = PendingActors[ActorIdx];
+				const FSoftObjectPath& OriginalActor = PendingActor.OriginActor;
+				const bool bIsActorObjectCompatible = IsRemappingCompatibleFunc(
+					OriginalActor, PendingActor.Class, TargetCandidate, TargetCandidate.ToSoftObjectPath()
+					);
+				if (bIsActorObjectCompatible && IsHierarchyCompatible(OriginalActor, TargetCandidate))
+				{
+					FActorWithSolutions& SolutionData = GetSolutionData(ActorIdx);
+					SolutionData.PossibleSolutions.Add(TargetCandidate.ToSoftObjectPath());
+				}
+			}
+		}
+
+		/** Validates that PossibleTargetActor's replicated hierarchy is compatible with that of OriginalActor. */
+		bool IsHierarchyCompatible(const FSoftObjectPath& OriginalActor, const FSoftObjectPtr& TargetCandidate) const
+		{
+			bool bIsHierarchyCompatible = true;
+			
+			ObjectPathHierarchy.TraverseTopToBottom([this, &TargetCandidate, &bIsHierarchyCompatible](const FChildRelation& Relation)
+			{
+				// E.g. hierarchy "Actor", "Actor.Foo", "Actor.Foo.Bar" where replication only maps "Actor" and "Actor.Foo.Bar".
+				// Then "Actor" and "Actor.Foo.Bar" are explicit, and "Actor.Foo" is implicit (needing no validation).
+				const bool bIsReplicated = Relation.Child.Type == EHierarchyObjectType::Explicit; 
+				if (!bIsReplicated)
+				{
+					return ETreeTraversalBehavior::Continue;
+				}
+
+				const FSoftObjectPath& OriginalObject = Relation.Child.Object;
+				const TOptional<FSoftObjectPath> TargetPath = ReplaceActorInPath(OriginalObject, TargetCandidate.ToSoftObjectPath());
+				if (!ensure(TargetPath.IsSet()))
+				{
+					// We don't expect this case to occur - so log it.
+					UE_LOG(LogConcert, Warning,
+						TEXT("Remapping: OriginalObject %s is no world object, or PossibleTargetActor %s is not an actor"),
+						*OriginalObject.ToString(),
+						*TargetCandidate.ToString()
+						);
+					return ETreeTraversalBehavior::Continue;
+				}
+				const FConcertReplicatedObjectInfo& ReplicationInfo = Origin.ReplicatedObjects[OriginalObject];
+				bIsHierarchyCompatible &= IsRemappingCompatibleFunc(OriginalObject, ReplicationInfo.ClassPath, TargetCandidate, *TargetPath);
+						
+				return bIsHierarchyCompatible ? ETreeTraversalBehavior::Continue : ETreeTraversalBehavior::Break;
+			}, OriginalActor);
+
+			return bIsHierarchyCompatible;
+		}
+		
+		/****************************** 3. Choose a solution and remap actor hierarchy to it ******************************/
+
+		/**
+		 * For each object that needs remapping, tries to pick a candidate we previously determined suitable.
+		 * 
+		 * The remapping is done greedily by always remapping the actor with least remaining options first with the goal of preventing accidental
+		 * starvation, see more explaination below.
+		 */
+		template<CProcessRemappingCallbable TProcessRemappingLambda>
+		void RemapActors(TProcessRemappingLambda&& ProcessRemapping)
+		{
+			while (!ObjectsWithSolutions.IsEmpty())
+			{
+				// Avoid starvation. Find index of the actor that has the least number of options remaining, i.e. is the most constrained.
+				// Example: OriginA could be remapped to TargetA or TargetB, OriginB can only be remapped to TargetB.
+				// We should pick an option for OriginB first. If instead we did OriginA first and unfortunately chose TargetB, OriginB would be
+				// left with 0 alternatives!
+				// Btw: if anything is left unassigned using this approach, then there was no way to satisfy all actors in the first place.
+				const int32 NextIndex = GetMostConstrainedIndex();
+				
+				const FActorWithSolutions& ObjectRemapData = ObjectsWithSolutions[NextIndex];
+				const FSoftObjectPath RemappedToActor = PickSolutionAndRemap(ObjectRemapData, ProcessRemapping);
+				ObjectsWithSolutions.RemoveAtSwap(NextIndex, EAllowShrinking::No);
+				
+				if (!RemappedToActor.IsNull())
+				{
+					RemoveFromSolutionSpace(RemappedToActor);
+				}
+			}
+		}
+
+		/** Finds the index of an actor that has the least amount of options left. */
+		int32 GetMostConstrainedIndex() const
+		{
+			int32 CurrentLowest = TNumericLimits<int32>::Max();
+			int32 Result = 0;
+			for (int32 i = 0; i < ObjectsWithSolutions.Num(); ++i)
+			{
+				const int32 NumRemaining = ObjectsWithSolutions[i].PossibleSolutions.Num();
+				if (NumRemaining == 1)
+				{
+					return i;
+				}
+				if (NumRemaining < CurrentLowest)
+				{
+					CurrentLowest = NumRemaining;
+					Result = i;
+				}
+			}
+			return Result;
+		}
+
+		/** Update ObjectsWithSolutions, removing TargetActor from the solution pool. Entries without solutions are removed.*/
+		void RemoveFromSolutionSpace(const FSoftObjectPath& TargetActor)
+		{
+			for (int32 i = 0; i < ObjectsWithSolutions.Num(); ++i)
+			{
+				TSmallInlineArray<FSoftObjectPath>& PossibleSolutions = ObjectsWithSolutions[i].PossibleSolutions;
+				if (!PossibleSolutions.Contains(TargetActor))
+				{
+					continue;
+				}
+				
+				if (PossibleSolutions.Num() == 1)
+				{
+					ObjectsWithSolutions.RemoveAtSwap(i);
+				}
+				else
+				{
+					PossibleSolutions.RemoveSingle(TargetActor);
+				}
+			}
+		}
+
+		/**
+		 * Picks the first solution, and replaces the original actor path to the new actor path.
+		 * Subobject paths are translated with the new actor name.
+		 */
+		template<CProcessRemappingCallbable TProcessRemappingLambda>
+		FSoftObjectPath PickSolutionAndRemap(const FActorWithSolutions& ObjectRemapData, TProcessRemappingLambda&& ProcessRemapping) const
+		{
+			if (ObjectRemapData.PossibleSolutions.IsEmpty())
+			{
+				return {};
+			}
+
+			const FSoftObjectPath& OriginActor = ObjectRemapData.OriginActor;
+			const FSoftObjectPath& PickedSolution = ObjectRemapData.PossibleSolutions[0];
+
+			// BuildOriginObjectHierarchyAndActorsNeedingRemapping added all replicated objects explicitly hence if an object is implicit, skip.
+			const TOptional<EHierarchyObjectType> Hierarchy = ObjectPathHierarchy.IsInHierarchy(ObjectRemapData.OriginActor);
+			const bool bIsReplicated = Hierarchy && *Hierarchy == EHierarchyObjectType::Explicit;
+			if (bIsReplicated)
+			{
+				ProcessRemapping(OriginActor, PickedSolution);
+			}
+			
+			RemapActorHierarchy(OriginActor, PickedSolution, ProcessRemapping);
+			
+			return PickedSolution;
+		}
+
+		/** For each original component that had replicated properties, replace the actor bit of their path with the chosen solution actor. */
+		template<CProcessRemappingCallbable TProcessRemappingLambda>
+		void RemapActorHierarchy(
+			const FSoftObjectPath& OriginActor,
+			const FSoftObjectPath& PickedSolution,
+			TProcessRemappingLambda&& ProcessRemapping
+			) const
+		{
+			ObjectPathHierarchy.TraverseTopToBottom([this, &ProcessRemapping, &PickedSolution](const FChildRelation& Relation)
+			{
+				const FSoftObjectPath& OriginalSubobject = Relation.Child.Object;
+				// BuildOriginObjectHierarchyAndActorsNeedingRemapping added all replicated objects explicitly hence if an object is implicit, skip.
+				const bool bNeedsRemapping = Relation.Child.Type == EHierarchyObjectType::Explicit;
+				const TOptional<FSoftObjectPath> RemappedSubobjectPath = bNeedsRemapping
+					? ReplaceActorInPath(OriginalSubobject, PickedSolution) : TOptional<FSoftObjectPath>{};
+				
+				if (RemappedSubobjectPath)
+				{
+					ProcessRemapping(OriginalSubobject, *RemappedSubobjectPath);
+				}
+				
+				return ETreeTraversalBehavior::Continue;
+			}, OriginActor);
+		}
+	};
+}
\ No newline at end of file
diff --git a/Engine/Plugins/Developer/Concert/ConcertSync/ConcertSyncTest/Source/ConcertSyncTest/Private/Replication/Components/LabelRemapping/DifferentWorld.spec.cpp b/Engine/Plugins/Developer/Concert/ConcertSync/ConcertSyncTest/Source/ConcertSyncTest/Private/Replication/Components/LabelRemapping/DifferentWorld.spec.cpp
new file mode 100644
--- /dev/null
+++ b/Engine/Plugins/Developer/Concert/ConcertSync/ConcertSyncTest/Source/ConcertSyncTest/Private/Replication/Components/LabelRemapping/DifferentWorld.spec.cpp
@@ -0,0 +1,211 @@
+﻿// Copyright Epic Games, Inc. All Rights Reserved.
+
+#include "Replication/Misc/ActorLabelRemappingEditor.h"
+#include "Replication/Util/ReplicatedTestWorld.h"
+
+#include "Components/PointLightComponent.h"
+#include "Components/StaticMeshComponent.h"
+#include "Engine/StaticMeshActor.h"
+#include "Misc/AutomationTest.h"
+
+namespace UE::ConcertSyncTests::Replication::ActorLabelRemapping
+{
+	/** This test that remapping works when the target world is different from the origin world. */
+	BEGIN_DEFINE_SPEC(FActorLabelRemappingDifferentWorldsSpec, "Editor.Concert.Replication.Components.ActorLabelRemapping.DifferentWorlds", EAutomationTestFlags::EditorContext | EAutomationTestFlags::EngineFilter)
+		TUniquePtr<FReplicatedTestWorld> OriginWorld;
+		TUniquePtr<FReplicatedTestWorld> TargetWorld;
+
+		UWorld& GetOriginWorld() const { return *OriginWorld->PreviewScene.GetWorld(); }
+		UWorld& GetTargetWorld() const { return *TargetWorld->PreviewScene.GetWorld(); }
+		
+		/** Tests that only PropertyName is replicated by Object in Map */
+		void TestReplicatesOnlyProperty(const FConcertObjectReplicationMap& Map, const UObject* Object, FName PropertyName)
+		{
+			Replication::TestReplicatesOnlyProperty(Map, Object, PropertyName, *this);
+		}
+	END_DEFINE_SPEC(FActorLabelRemappingDifferentWorldsSpec);
+
+	void FActorLabelRemappingDifferentWorldsSpec::Define()
+	{
+		BeforeEach([this]
+		{
+			OriginWorld = MakeUnique<FReplicatedTestWorld>();
+			TargetWorld = MakeUnique<FReplicatedTestWorld>();
+		});
+		AfterEach([this]
+		{
+			OriginWorld.Reset();
+			TargetWorld.Reset();
+		});
+		
+		It("Remap actor that retains label and object path", [this]
+		{
+			// 1. Setup objects
+			AStaticMeshActor* OriginActor = OriginWorld->SpawnActor<AStaticMeshActor>(TEXT("StaticMeshActor0_Origin"), TEXT("Label0"));
+			UStaticMeshComponent* OriginComponent = OriginActor->GetStaticMeshComponent();
+			OriginWorld->AddReplicatedProperty(OriginActor, GET_MEMBER_NAME_CHECKED(AActor, bNetTemporary));
+			OriginWorld->AddReplicatedProperty(OriginComponent, GET_MEMBER_NAME_CHECKED(USceneComponent, ComponentVelocity));
+
+			AStaticMeshActor* TargetActor = TargetWorld->SpawnActor<AStaticMeshActor>(TEXT("StaticMeshActor0_Target"), TEXT("Label0"));
+			UStaticMeshComponent* TargetComponent = TargetActor->GetStaticMeshComponent();
+			
+			// 2. Generate data
+			const FConcertReplicationRemappingData RemappingData = ConcertSyncCore::GenerateRemappingData(OriginWorld->ReplicationMap);
+			const FConcertObjectReplicationMap Translation = ConcertSyncCore::RemapReplicationMap(OriginWorld->ReplicationMap, RemappingData, GetTargetWorld());
+
+			// 3. Test
+			TestEqual(TEXT("ReplicatedObjects.Num()"), Translation.ReplicatedObjects.Num(), 2);
+			TestReplicatesOnlyProperty(Translation, TargetActor, GET_MEMBER_NAME_CHECKED(AActor, bNetTemporary));
+			TestReplicatesOnlyProperty(Translation, TargetComponent, GET_MEMBER_NAME_CHECKED(USceneComponent, ComponentVelocity));
+		});
+
+		It("Remap component of which the actor is not replicated", [this]
+		{
+			// 1. Setup objects
+			AStaticMeshActor* OriginActor = OriginWorld->SpawnActor<AStaticMeshActor>(TEXT("StaticMeshActor0_Origin"), TEXT("Label0"));
+			UStaticMeshComponent* OriginComponent = OriginActor->GetStaticMeshComponent();
+			OriginWorld->AddReplicatedProperty(OriginComponent, GET_MEMBER_NAME_CHECKED(USceneComponent, ComponentVelocity));
+			
+			AStaticMeshActor* TargetActor = TargetWorld->SpawnActor<AStaticMeshActor>(TEXT("StaticMeshActor0_Target"), TEXT("Label0"));
+			UStaticMeshComponent* TargetComponent = TargetActor->GetStaticMeshComponent();
+			
+			// 2. Generate data
+			const FConcertReplicationRemappingData RemappingData = ConcertSyncCore::GenerateRemappingData(OriginWorld->ReplicationMap);
+			const FConcertObjectReplicationMap Translation = ConcertSyncCore::RemapReplicationMap(OriginWorld->ReplicationMap, RemappingData, GetTargetWorld());
+
+			// 3. Test
+			TestEqual(TEXT("ReplicatedObjects.Num()"), Translation.ReplicatedObjects.Num(), 1);
+			TestReplicatesOnlyProperty(Translation, TargetComponent, GET_MEMBER_NAME_CHECKED(USceneComponent, ComponentVelocity));
+		});
+
+		It("Do not remap to actor that does not have the label", [this]
+		{
+			// To add potential risk, we'll make both actors have the same name.
+			const FName ActorName(TEXT("StatisMeshActor0"));
+			
+			// 1. Setup objects
+			AStaticMeshActor* OriginActor = OriginWorld->SpawnActor<AStaticMeshActor>(ActorName, TEXT("Label0"));
+			UStaticMeshComponent* OriginComponent = OriginActor->GetStaticMeshComponent();
+			OriginWorld->AddReplicatedProperty(OriginComponent, GET_MEMBER_NAME_CHECKED(USceneComponent, ComponentVelocity));
+			
+			TargetWorld->SpawnActor<AStaticMeshActor>(ActorName, TEXT("Label0_Other"));
+			
+			// 2. Generate data
+			const FConcertReplicationRemappingData RemappingData = ConcertSyncCore::GenerateRemappingData(OriginWorld->ReplicationMap);
+			const FConcertObjectReplicationMap Translation = ConcertSyncCore::RemapReplicationMap(OriginWorld->ReplicationMap, RemappingData, GetTargetWorld());
+
+			// 3. Test
+			TestEqual(TEXT("ReplicatedObjects.Num()"), Translation.ReplicatedObjects.Num(), 0);
+		});
+		
+		It("Correctly remap actors that swapped labels", [this]
+		{
+			// 1. Setup objects
+			AStaticMeshActor* OriginCross0 = OriginWorld->SpawnActor<AStaticMeshActor>(TEXT("Cross0"), TEXT("Label0"));
+			AStaticMeshActor* OriginCross1 = OriginWorld->SpawnActor<AStaticMeshActor>(TEXT("Cross1"), TEXT("Label1"));
+			UStaticMeshComponent* OriginComponent0 = OriginCross0->GetStaticMeshComponent();
+			UStaticMeshComponent* OriginComponent1 = OriginCross1->GetStaticMeshComponent();
+			OriginWorld->AddReplicatedProperty(OriginCross0, GET_MEMBER_NAME_CHECKED(AActor, bNetTemporary));
+			OriginWorld->AddReplicatedProperty(OriginComponent0, GET_MEMBER_NAME_CHECKED(USceneComponent, ComponentVelocity));
+			OriginWorld->AddReplicatedProperty(OriginCross1, GET_MEMBER_NAME_CHECKED(AActor, bOnlyRelevantToOwner));
+			OriginWorld->AddReplicatedProperty(OriginComponent1, GET_MEMBER_NAME_CHECKED(USceneComponent, bHiddenInGame));
+			
+			AStaticMeshActor* TargetCross0 = TargetWorld->SpawnActor<AStaticMeshActor>(TEXT("Cross0"), TEXT("Label1"));
+			AStaticMeshActor* TargetCross1 = TargetWorld->SpawnActor<AStaticMeshActor>(TEXT("Cross1"), TEXT("Label0"));
+			UStaticMeshComponent* TargetComponent0 = TargetCross0->GetStaticMeshComponent();
+			UStaticMeshComponent* TargetComponent1 = TargetCross1->GetStaticMeshComponent();
+			
+			// 2. Generate data
+			const FConcertReplicationRemappingData RemappingData = ConcertSyncCore::GenerateRemappingData(OriginWorld->ReplicationMap);
+			const FConcertObjectReplicationMap Translation = ConcertSyncCore::RemapReplicationMap(OriginWorld->ReplicationMap, RemappingData, GetTargetWorld());
+
+			// 3. 
+			TestEqual(TEXT("ReplicatedObjects.Num()"), Translation.ReplicatedObjects.Num(), 4);
+			// ... swaps the assigned properties
+			TestReplicatesOnlyProperty(Translation, TargetCross0, GET_MEMBER_NAME_CHECKED(AActor, bOnlyRelevantToOwner));
+			TestReplicatesOnlyProperty(Translation, TargetComponent0, GET_MEMBER_NAME_CHECKED(USceneComponent, bHiddenInGame));
+			TestReplicatesOnlyProperty(Translation, TargetCross1, GET_MEMBER_NAME_CHECKED(AActor, bNetTemporary));
+			TestReplicatesOnlyProperty(Translation, TargetComponent1, GET_MEMBER_NAME_CHECKED(USceneComponent, ComponentVelocity));
+		});
+		
+		It("Actor is not remapped if owned, replicated component class changes", [this]
+		{
+			// 1. Setup objects
+			const FName ComponentName(TEXT("Component"));
+			AActor* OriginActor = OriginWorld->SpawnActor<AActor>(TEXT("Actor0"), TEXT("Label0"));
+			UStaticMeshComponent* OriginComponent = NewObject<UStaticMeshComponent>(OriginActor, ComponentName);
+			OriginActor->AddOwnedComponent(OriginComponent);
+			OriginWorld->AddReplicatedProperty(OriginActor, GET_MEMBER_NAME_CHECKED(AActor, bNetTemporary));
+			OriginWorld->AddReplicatedProperty(OriginComponent, GET_MEMBER_NAME_CHECKED(USceneComponent, ComponentVelocity));
+			
+			AActor* TargetActor = TargetWorld->SpawnActor<AActor>(TEXT("Actor0"), TEXT("Label0"));
+			UPointLightComponent* TargetComponent = NewObject<UPointLightComponent>(TargetActor, ComponentName);
+			TargetActor->AddOwnedComponent(TargetComponent);
+			
+			// 2. Generate data
+			const FConcertReplicationRemappingData RemappingData = ConcertSyncCore::GenerateRemappingData(OriginWorld->ReplicationMap);
+			const FConcertObjectReplicationMap Translation = ConcertSyncCore::RemapReplicationMap(OriginWorld->ReplicationMap, RemappingData, GetTargetWorld());
+
+			// 3. Test
+			TestEqual(TEXT("ReplicatedObjects.Num()"), Translation.ReplicatedObjects.Num(), 0);
+		});
+		It("Actor matched if owned, non-replicated component class changes", [this]
+		{
+			// 1. Setup objects
+			const FName ComponentName(TEXT("Component"));
+			AActor* OriginActor = OriginWorld->SpawnActor<AActor>(TEXT("Actor0"), TEXT("Label0"));
+			UStaticMeshComponent* OriginComponent = NewObject<UStaticMeshComponent>(OriginActor, ComponentName);
+			OriginActor->AddOwnedComponent(OriginComponent);
+			OriginWorld->AddReplicatedProperty(OriginActor, GET_MEMBER_NAME_CHECKED(AActor, bNetTemporary));
+			
+			AActor* TargetActor = TargetWorld->SpawnActor<AActor>(TEXT("Actor0"), TEXT("Label0"));
+			UPointLightComponent* TargetComponent = NewObject<UPointLightComponent>(TargetActor, ComponentName);
+			TargetActor->AddOwnedComponent(TargetComponent);
+			
+			// 2. Generate data
+			const FConcertReplicationRemappingData RemappingData = ConcertSyncCore::GenerateRemappingData(OriginWorld->ReplicationMap);
+			const FConcertObjectReplicationMap Translation = ConcertSyncCore::RemapReplicationMap(OriginWorld->ReplicationMap, RemappingData, GetTargetWorld());
+
+			// 3. Test
+			TestEqual(TEXT("ReplicatedObjects.Num()"), Translation.ReplicatedObjects.Num(), 1);
+			TestReplicatesOnlyProperty(Translation, TargetActor, GET_MEMBER_NAME_CHECKED(AActor, bNetTemporary));
+		});
+
+		It("Actor is not remapped if owned, replicated component name changes", [this]
+		{
+			// 1. Setup objects
+			AActor* OriginActor = OriginWorld->SpawnActor<AActor>(TEXT("Actor0"), TEXT("Label0"));
+			UStaticMeshComponent* OriginComponent = NewObject<UStaticMeshComponent>(OriginActor, TEXT("OriginComponent"));
+			OriginActor->AddOwnedComponent(OriginComponent);
+			OriginWorld->AddReplicatedProperty(OriginActor, GET_MEMBER_NAME_CHECKED(AActor, bNetTemporary));
+			OriginWorld->AddReplicatedProperty(OriginComponent, GET_MEMBER_NAME_CHECKED(USceneComponent, ComponentVelocity));
+			
+			AActor* TargetActor = TargetWorld->SpawnActor<AActor>(TEXT("Actor0"), TEXT("Label0"));
+			UPointLightComponent* TargetComponent = NewObject<UPointLightComponent>(TargetActor, TEXT("TargetComponent"));
+			TargetActor->AddOwnedComponent(TargetComponent);
+			
+			// 2. Generate data
+			const FConcertReplicationRemappingData RemappingData = ConcertSyncCore::GenerateRemappingData(OriginWorld->ReplicationMap);
+			const FConcertObjectReplicationMap Translation = ConcertSyncCore::RemapReplicationMap(OriginWorld->ReplicationMap, RemappingData, GetTargetWorld());
+
+			// 3. Test
+			TestEqual(TEXT("ReplicatedObjects.Num()"), Translation.ReplicatedObjects.Num(), 0);
+		});
+
+		It("Cannot remap onto actor that has changed class but has same name and label", [this]
+		{
+			// 1. Setup objects
+			AActor* OriginActor = OriginWorld->SpawnActor<AActor>(TEXT("Actor0"), TEXT("Label0"));
+			OriginWorld->AddReplicatedProperty(OriginActor, GET_MEMBER_NAME_CHECKED(AActor, bNetTemporary));
+			
+			TargetWorld->SpawnActor<AStaticMeshActor>(TEXT("Actor0"), TEXT("Label0"));
+			
+			// 2. Generate data
+			const FConcertReplicationRemappingData RemappingData = ConcertSyncCore::GenerateRemappingData(OriginWorld->ReplicationMap);
+			const FConcertObjectReplicationMap Translation = ConcertSyncCore::RemapReplicationMap(OriginWorld->ReplicationMap, RemappingData, GetTargetWorld());
+
+			// 3. Test
+			TestEqual(TEXT("ReplicatedObjects.Num()"), Translation.ReplicatedObjects.Num(), 0);
+		});
+	}
+}
diff --git a/Engine/Plugins/Developer/Concert/ConcertSync/ConcertSyncTest/Source/ConcertSyncTest/Private/Replication/Components/LabelRemapping/GenerateRemappingData.spec.cpp b/Engine/Plugins/Developer/Concert/ConcertSync/ConcertSyncTest/Source/ConcertSyncTest/Private/Replication/Components/LabelRemapping/GenerateRemappingData.spec.cpp
new file mode 100644
--- /dev/null
+++ b/Engine/Plugins/Developer/Concert/ConcertSync/ConcertSyncTest/Source/ConcertSyncTest/Private/Replication/Components/LabelRemapping/GenerateRemappingData.spec.cpp
@@ -0,0 +1,113 @@
+﻿// Copyright Epic Games, Inc. All Rights Reserved.
+
+#include "Replication/Misc/ActorLabelRemappingEditor.h"
+#include "Replication/Util/ReplicatedTestWorld.h"
+
+#include "Components/StaticMeshComponent.h"
+#include "Engine/StaticMeshActor.h"
+#include "Misc/AutomationTest.h"
+
+namespace UE::ConcertSyncTests::Replication::ActorLabelRemapping
+{
+	/** This test that remapping works after changes have been made to the world. */
+	BEGIN_DEFINE_SPEC(FGenerateRemappingDataSpec, "Editor.Concert.Replication.Components.ActorLabelRemapping.GenerateRemappingDataSpec", EAutomationTestFlags::EditorContext | EAutomationTestFlags::EngineFilter)
+		TUniquePtr<FReplicatedTestWorld> ReplicatedWorld;
+
+		UWorld& GetWorld() const { return *ReplicatedWorld->PreviewScene.GetWorld(); }
+		
+		/** Tests that only PropertyName is replicated by Object in Map */
+		void TestReplicatesOnlyProperty(const FConcertObjectReplicationMap& Map, const UObject* Object, FName PropertyName)
+		{
+			Replication::TestReplicatesOnlyProperty(Map, Object, PropertyName, *this);
+		}
+	END_DEFINE_SPEC(FGenerateRemappingDataSpec);
+	
+	void FGenerateRemappingDataSpec::Define()
+	{
+		BeforeEach([this]
+		{
+			ReplicatedWorld = MakeUnique<FReplicatedTestWorld>();
+		});
+		AfterEach([this]
+		{
+			ReplicatedWorld.Reset();
+		});
+		
+		It("Contains actor data when actor is replicated", [this]
+		{
+			// 1. Create objects
+			const FString ActorLabel = TEXT("Label0");
+			AStaticMeshActor* Actor = ReplicatedWorld->SpawnActor<AStaticMeshActor>(TEXT("StaticMeshActor0"), ActorLabel);
+			ReplicatedWorld->AddReplicatedProperty(Actor, GET_MEMBER_NAME_CHECKED(AActor, bNetTemporary));
+
+			// 2. Generate data
+			const FConcertReplicationRemappingData RemappingData = ConcertSyncCore::GenerateRemappingData(ReplicatedWorld->ReplicationMap);
+			const FConcertReplicationRemappingData_Actor* ActorData = RemappingData.ActorData.Find(Actor);
+			
+			// 3. Test data
+			if (!ActorData)
+			{
+				AddError(TEXT("Missing data"));
+			}
+			else
+			{
+				TestEqual(TEXT("Label"), ActorData->Label, ActorLabel);
+				TestEqual(TEXT("Class"), ActorData->Class, FSoftClassPath(Actor->GetClass()));
+			}
+		});
+
+		It("Contains actor data even if only component is replicated", [this]
+		{
+			// 1. Create objects
+			const FString ActorLabel = TEXT("Label0");
+			AStaticMeshActor* Actor = ReplicatedWorld->SpawnActor<AStaticMeshActor>(TEXT("StaticMeshActor0"), ActorLabel);
+			UStaticMeshComponent* Comp = Actor->GetStaticMeshComponent();
+			ReplicatedWorld->AddReplicatedProperty(Comp, GET_MEMBER_NAME_CHECKED(USceneComponent, ComponentVelocity));
+			
+			// 2. Generate data
+			const FConcertReplicationRemappingData RemappingData = ConcertSyncCore::GenerateRemappingData(ReplicatedWorld->ReplicationMap);
+			const FConcertReplicationRemappingData_Actor* ActorData = RemappingData.ActorData.Find(Actor);
+			
+			// 3. Test data
+			if (!ActorData)
+			{
+				AddError(TEXT("Missing data"));
+			}
+			else
+			{
+				TestEqual(TEXT("Label"), ActorData->Label, ActorLabel);
+				TestEqual(TEXT("Class"), ActorData->Class, FSoftClassPath(Actor->GetClass()));
+			}
+		});
+		
+		It("Contains no actor data if component's owning actor class cannot be gotten", [this]
+		{
+			// 1. Create objects
+			const FString ActorLabel = TEXT("Label0");
+			const FSoftObjectPath ActorPath(TEXT("/InternalTest/Map.Map:PersistentLevel.Actor"));
+			const FSoftObjectPath ComponentPath(TEXT("/InternalTest/Map.Map:PersistentLevel.Actor.Component"));
+			ensureAlways(ActorPath.ResolveObject() == nullptr);
+			
+			FConcertObjectReplicationMap ReplicationMap;
+			FConcertReplicatedObjectInfo& Info = ReplicationMap.ReplicatedObjects.Add(ComponentPath);
+			Info.ClassPath = USceneComponent::StaticClass();
+			const TOptional<FConcertPropertyChain> Property = FConcertPropertyChain::CreateFromPath(
+				*USceneComponent::StaticClass(),
+				{ GET_MEMBER_NAME_CHECKED(USceneComponent, ComponentVelocity) }
+				);
+			if (!Property)
+			{
+				AddError(TEXT("Failed get FConcertPropertyChain for property"));
+				return;
+			}
+			Info.PropertySelection.ReplicatedProperties.Add(*Property);
+			
+			// 2. Generate data
+			// Internally, GenerateRemappingData will try to resolve ActorPath but fail because the instance does not exist.
+			const FConcertReplicationRemappingData RemappingData = ConcertSyncCore::GenerateRemappingData(ReplicatedWorld->ReplicationMap);
+			
+			// 3. Test data
+			TestFalse(TEXT("Does not contain data"), RemappingData.ActorData.Contains(ActorPath));
+		});
+	}
+}
diff --git a/Engine/Plugins/Developer/Concert/ConcertSync/ConcertSyncTest/Source/ConcertSyncTest/Private/Replication/Components/LabelRemapping/SameWorld.spec.cpp b/Engine/Plugins/Developer/Concert/ConcertSync/ConcertSyncTest/Source/ConcertSyncTest/Private/Replication/Components/LabelRemapping/SameWorld.spec.cpp
new file mode 100644
--- /dev/null
+++ b/Engine/Plugins/Developer/Concert/ConcertSync/ConcertSyncTest/Source/ConcertSyncTest/Private/Replication/Components/LabelRemapping/SameWorld.spec.cpp
@@ -0,0 +1,223 @@
+﻿// Copyright Epic Games, Inc. All Rights Reserved.
+
+#include "Replication/Misc/ActorLabelRemappingEditor.h"
+#include "Replication/Util/ReplicatedTestWorld.h"
+
+#include "Components/PointLightComponent.h"
+#include "Components/StaticMeshComponent.h"
+#include "Engine/BlockingVolume.h"
+#include "Engine/StaticMeshActor.h"
+#include "GameFramework/CameraBlockingVolume.h"
+#include "Misc/AutomationTest.h"
+
+namespace UE::ConcertSyncTests::Replication::ActorLabelRemapping
+{
+	/** This test that remapping works after changes have been made to the world. */
+	BEGIN_DEFINE_SPEC(FActorLabelRemappingSameWorldSpec, "Editor.Concert.Replication.Components.ActorLabelRemapping.SameWorld", EAutomationTestFlags::EditorContext | EAutomationTestFlags::EngineFilter)
+		TUniquePtr<FReplicatedTestWorld> ReplicatedWorld;
+
+		UWorld& GetWorld() const { return *ReplicatedWorld->PreviewScene.GetWorld(); }
+	
+		/** Tests that only PropertyName is replicated by Object in Map */
+		void TestReplicatesOnlyProperty(const FConcertObjectReplicationMap& Map, const UObject* Object, FName PropertyName)
+		{
+			Replication::TestReplicatesOnlyProperty(Map, Object, PropertyName, *this);
+		}
+	END_DEFINE_SPEC(FActorLabelRemappingSameWorldSpec);
+
+	void FActorLabelRemappingSameWorldSpec::Define()
+	{
+		BeforeEach([this]
+		{
+			ReplicatedWorld = MakeUnique<FReplicatedTestWorld>();
+		});
+		AfterEach([this]
+		{
+			ReplicatedWorld.Reset();
+		});
+
+		It("Remap actor that retains label and object path", [this]
+		{
+			// 1. Setup objects
+			AStaticMeshActor* Actor = ReplicatedWorld->SpawnActor<AStaticMeshActor>(TEXT("StaticMeshActor0"), TEXT("Label0"));
+			UStaticMeshComponent* Component = Actor->GetStaticMeshComponent();
+			ReplicatedWorld->AddReplicatedProperty(Actor, GET_MEMBER_NAME_CHECKED(AActor, bNetTemporary));
+			ReplicatedWorld->AddReplicatedProperty(Component, GET_MEMBER_NAME_CHECKED(USceneComponent, ComponentVelocity));
+			
+			// 2. Generate data
+			const FConcertReplicationRemappingData RemappingData = ConcertSyncCore::GenerateRemappingData(ReplicatedWorld->ReplicationMap);
+			const FConcertObjectReplicationMap Translation = ConcertSyncCore::RemapReplicationMap(ReplicatedWorld->ReplicationMap, RemappingData, GetWorld());
+
+			// 3. Test
+			TestEqual(TEXT("ReplicatedObjects.Num()"), Translation.ReplicatedObjects.Num(), 2);
+			TestReplicatesOnlyProperty(Translation, Actor, GET_MEMBER_NAME_CHECKED(AActor, bNetTemporary));
+			TestReplicatesOnlyProperty(Translation, Component, GET_MEMBER_NAME_CHECKED(USceneComponent, ComponentVelocity));
+		});
+
+		It("Remap component of which the actor is not replicated", [this]
+		{
+			// 1. Setup objects
+			AStaticMeshActor* Actor = ReplicatedWorld->SpawnActor<AStaticMeshActor>(TEXT("StaticMeshActor0"), TEXT("Label0"));
+			UStaticMeshComponent* Component = Actor->GetStaticMeshComponent();
+			ReplicatedWorld->AddReplicatedProperty(Component, GET_MEMBER_NAME_CHECKED(USceneComponent, ComponentVelocity));
+			
+			// 2. Generate data
+			const FConcertReplicationRemappingData RemappingData = ConcertSyncCore::GenerateRemappingData(ReplicatedWorld->ReplicationMap);
+			const FConcertObjectReplicationMap Translation = ConcertSyncCore::RemapReplicationMap(ReplicatedWorld->ReplicationMap, RemappingData, GetWorld());
+
+			// 3. Test
+			TestEqual(TEXT("ReplicatedObjects.Num()"), Translation.ReplicatedObjects.Num(), 1);
+			TestReplicatesOnlyProperty(Translation, Component, GET_MEMBER_NAME_CHECKED(USceneComponent, ComponentVelocity));
+		});
+
+		It("Do not remap to actor that no longer has the label", [this]
+		{
+			// 1. Setup objects
+			AStaticMeshActor* Actor = ReplicatedWorld->SpawnActor<AStaticMeshActor>(TEXT("StaticMeshActor0"), TEXT("Label0"));
+			ReplicatedWorld->AddReplicatedProperty(Actor, GET_MEMBER_NAME_CHECKED(AActor, bNetTemporary));
+			
+			// 2. Generate data
+			const FConcertReplicationRemappingData RemappingData = ConcertSyncCore::GenerateRemappingData(ReplicatedWorld->ReplicationMap);
+			Actor->SetActorLabel(TEXT("ChangedLabel"));
+			const FConcertObjectReplicationMap Translation = ConcertSyncCore::RemapReplicationMap(ReplicatedWorld->ReplicationMap, RemappingData, GetWorld());
+
+			// 3. Test
+			TestEqual(TEXT("ReplicatedObjects.Num()"), Translation.ReplicatedObjects.Num(), 0);
+		});
+		
+		It("Correctly remap actors that swapped labels", [this]
+		{
+			// 1. Setup objects
+			AStaticMeshActor* Cross0 = ReplicatedWorld->SpawnActor<AStaticMeshActor>(TEXT("Cross0"), TEXT("Label0"));
+			AStaticMeshActor* Cross1 = ReplicatedWorld->SpawnActor<AStaticMeshActor>(TEXT("Cross1"), TEXT("Label1"));
+			UStaticMeshComponent* Component0 = Cross0->GetStaticMeshComponent();
+			UStaticMeshComponent* Component1 = Cross1->GetStaticMeshComponent();
+			ReplicatedWorld->AddReplicatedProperty(Cross0, GET_MEMBER_NAME_CHECKED(AActor, bNetTemporary));
+			ReplicatedWorld->AddReplicatedProperty(Component0, GET_MEMBER_NAME_CHECKED(USceneComponent, ComponentVelocity));
+			ReplicatedWorld->AddReplicatedProperty(Cross1, GET_MEMBER_NAME_CHECKED(AActor, bOnlyRelevantToOwner));
+			ReplicatedWorld->AddReplicatedProperty(Component1, GET_MEMBER_NAME_CHECKED(USceneComponent, bHiddenInGame));
+			
+			// 2. Generate data
+			const FConcertReplicationRemappingData RemappingData = ConcertSyncCore::GenerateRemappingData(ReplicatedWorld->ReplicationMap);
+			// Swapping labels ...
+			Cross0->SetActorLabel(TEXT("Label1"));
+			Cross1->SetActorLabel(TEXT("Label0"));
+			const FConcertObjectReplicationMap Translation = ConcertSyncCore::RemapReplicationMap(ReplicatedWorld->ReplicationMap, RemappingData, GetWorld());
+
+			// 3. 
+			TestEqual(TEXT("ReplicatedObjects.Num()"), Translation.ReplicatedObjects.Num(), 4);
+			// ... swaps the assigned properties
+			TestReplicatesOnlyProperty(Translation, Cross0, GET_MEMBER_NAME_CHECKED(AActor, bOnlyRelevantToOwner));
+			TestReplicatesOnlyProperty(Translation, Component0, GET_MEMBER_NAME_CHECKED(USceneComponent, bHiddenInGame));
+			TestReplicatesOnlyProperty(Translation, Cross1, GET_MEMBER_NAME_CHECKED(AActor, bNetTemporary));
+			TestReplicatesOnlyProperty(Translation, Component1, GET_MEMBER_NAME_CHECKED(USceneComponent, ComponentVelocity));
+		});
+
+		It("Actor is not remapped if owned, replicated component class changes", [this]
+		{
+			// 1. Setup objects
+			const FName ComponentName(TEXT("Component"));
+			AActor* Actor = ReplicatedWorld->SpawnActor<AActor>(TEXT("Actor0"), TEXT("Label0"));
+			UStaticMeshComponent* PreComponent = NewObject<UStaticMeshComponent>(Actor, ComponentName);
+			Actor->AddOwnedComponent(PreComponent);
+			ReplicatedWorld->AddReplicatedProperty(Actor, GET_MEMBER_NAME_CHECKED(AActor, bNetTemporary));
+			ReplicatedWorld->AddReplicatedProperty(PreComponent, GET_MEMBER_NAME_CHECKED(USceneComponent, ComponentVelocity));
+			
+			// 2. Generate data
+			const FConcertReplicationRemappingData RemappingData = ConcertSyncCore::GenerateRemappingData(ReplicatedWorld->ReplicationMap);
+			
+			PreComponent->Rename(TEXT("TRASH"));
+			UPointLightComponent* PostComponent = NewObject<UPointLightComponent>(Actor, ComponentName);
+			Actor->AddOwnedComponent(PostComponent);
+			
+			const FConcertObjectReplicationMap Translation = ConcertSyncCore::RemapReplicationMap(ReplicatedWorld->ReplicationMap, RemappingData, GetWorld());
+
+			// 3. Test
+			TestEqual(TEXT("ReplicatedObjects.Num()"), Translation.ReplicatedObjects.Num(), 0);
+		});
+		It("Actor matched if owned, non-replicated component class changes", [this]
+		{
+			// 1. Setup objects
+			const FName ComponentName(TEXT("Component"));
+			AActor* Actor = ReplicatedWorld->SpawnActor<AActor>(TEXT("Actor0"), TEXT("Label0"));
+			UStaticMeshComponent* PreComponent = NewObject<UStaticMeshComponent>(Actor, ComponentName);
+			Actor->AddOwnedComponent(PreComponent);
+			ReplicatedWorld->AddReplicatedProperty(Actor, GET_MEMBER_NAME_CHECKED(AActor, bNetTemporary));
+				
+			// 2. Generate data
+			const FConcertReplicationRemappingData RemappingData = ConcertSyncCore::GenerateRemappingData(ReplicatedWorld->ReplicationMap);
+			
+			PreComponent->Rename(TEXT("TRASH"));
+			UPointLightComponent* PostComponent = NewObject<UPointLightComponent>(Actor, ComponentName);
+			Actor->AddOwnedComponent(PostComponent);
+			
+			const FConcertObjectReplicationMap Translation = ConcertSyncCore::RemapReplicationMap(ReplicatedWorld->ReplicationMap, RemappingData, GetWorld());
+
+			// 3. Test
+			TestEqual(TEXT("ReplicatedObjects.Num()"), Translation.ReplicatedObjects.Num(), 1);
+			TestReplicatesOnlyProperty(Translation, Actor, GET_MEMBER_NAME_CHECKED(AActor, bNetTemporary));
+		});
+
+		It("Actor is not remapped if owned, replicated component name changes", [this]
+		{
+			// 1. Setup objects
+			AStaticMeshActor* Actor = ReplicatedWorld->SpawnActor<AStaticMeshActor>(TEXT("StaticMeshActor0"), TEXT("Label0"));
+			UStaticMeshComponent* Component = Actor->GetStaticMeshComponent();
+			ReplicatedWorld->AddReplicatedProperty(Actor, GET_MEMBER_NAME_CHECKED(AActor, bNetTemporary));
+			ReplicatedWorld->AddReplicatedProperty(Component, GET_MEMBER_NAME_CHECKED(USceneComponent, ComponentVelocity));
+			
+			// 2. Generate data
+			const FConcertReplicationRemappingData RemappingData = ConcertSyncCore::GenerateRemappingData(ReplicatedWorld->ReplicationMap);
+			Component->Rename(TEXT("NewComponentName"));
+			const FConcertObjectReplicationMap Translation = ConcertSyncCore::RemapReplicationMap(ReplicatedWorld->ReplicationMap, RemappingData, GetWorld());
+
+			// 3. Test
+			TestEqual(TEXT("ReplicatedObjects.Num()"), Translation.ReplicatedObjects.Num(), 0);
+		});
+
+		Describe("Cannot remap onto actor that has changed class but has same name and label", [this]
+		{
+			const auto RunVariation = [this](bool bActorHasProperties, bool bCreateComponent)
+			{
+				const FString Label(TEXT("Label0"));
+				const FName ActorName(TEXT("Actor0"));
+				const FName ComponentName(TEXT("Component"));
+				
+				// 1. Setup objects: ACameraBlockingVolume and ABlockingVolume were chosen because their classes have the same component hierarchy
+				// to avoid a potential, unlikely point of failure for the test set up. Added bonus: they don't inherit from each other.
+				AVolume* Actor = ReplicatedWorld->SpawnActor<ABlockingVolume>(ActorName, Label);
+				if (bActorHasProperties)
+				{
+					ReplicatedWorld->AddReplicatedProperty(Actor, GET_MEMBER_NAME_CHECKED(AActor, bNetTemporary));
+				}
+				if (bCreateComponent)
+				{
+					UStaticMeshComponent* Component = NewObject<UStaticMeshComponent>(Actor, ComponentName);
+					Actor->AddOwnedComponent(Component);
+					ReplicatedWorld->AddReplicatedProperty(Component, GET_MEMBER_NAME_CHECKED(USceneComponent, ComponentVelocity));
+				}
+				
+				// 2. Generate data
+				const FConcertReplicationRemappingData RemappingData = ConcertSyncCore::GenerateRemappingData(ReplicatedWorld->ReplicationMap);
+				Actor->Rename(TEXT("TRASH"));
+				// Avoid the old actor being considered for remapping.
+				Actor->SetActorLabel(TEXT("TRASH"));
+				Actor->Destroy();
+				Actor = ReplicatedWorld->SpawnActor<ACameraBlockingVolume>(ActorName, Label);
+				if (bCreateComponent)
+				{
+					UStaticMeshComponent* Component = NewObject<UStaticMeshComponent>(Actor, ComponentName);
+					Actor->AddOwnedComponent(Component);
+				}
+				const FConcertObjectReplicationMap Translation = ConcertSyncCore::RemapReplicationMap(ReplicatedWorld->ReplicationMap, RemappingData, GetWorld());
+				
+				// 3. Test
+				TestEqual(TEXT("ReplicatedObjects.Num()"), Translation.ReplicatedObjects.Num(), 0);
+			};
+
+			It("Only actor has properties", [RunVariation] { RunVariation(true, false); });
+			It("Only component has properties", [RunVariation] { RunVariation(false, true); });
+			It("Actor and component have properties", [RunVariation] { RunVariation(true, true); });
+		});
+	}
+}
diff --git a/Engine/Plugins/Developer/Concert/ConcertSync/ConcertSyncTest/Source/ConcertSyncTest/Private/Replication/Util/ReplicatedTestWorld.h b/Engine/Plugins/Developer/Concert/ConcertSync/ConcertSyncTest/Source/ConcertSyncTest/Private/Replication/Util/ReplicatedTestWorld.h
new file mode 100644
--- /dev/null
+++ b/Engine/Plugins/Developer/Concert/ConcertSync/ConcertSyncTest/Source/ConcertSyncTest/Private/Replication/Util/ReplicatedTestWorld.h
@@ -0,0 +1,88 @@
+﻿// Copyright Epic Games, Inc. All Rights Reserved.
+
+#pragma once
+
+#include "ConcertLogGlobal.h"
+#include "Replication/Misc/ActorLabelRemappingCore.h"
+
+#include "Algo/AnyOf.h"
+#include "Algo/Transform.h"
+#include "PreviewScene.h"
+#include "Misc/AutomationTest.h"
+
+namespace UE::ConcertSyncTests::Replication
+{
+	/** Util that has a world and replication map of the objects in it. */
+	struct FReplicatedTestWorld
+	{
+		FPreviewScene PreviewScene;
+		FConcertObjectReplicationMap ReplicationMap;
+
+		FReplicatedTestWorld()
+			: PreviewScene(FPreviewScene::ConstructionValues()
+				.SetTransactional(false)
+				.SetCreateDefaultLighting(false)
+				.ShouldSimulatePhysics(false)
+				.SetCreatePhysicsScene(false)
+			)
+		{}
+
+		template<typename TActor> 
+		TActor* SpawnActor(FName Name, FString Label)
+		{
+			FActorSpawnParameters Params;
+			Params.Name = Name;
+			Params.NameMode = FActorSpawnParameters::ESpawnActorNameMode::Required_ErrorAndReturnNull;
+			TActor* Actor = PreviewScene.GetWorld()->SpawnActor<TActor>(Params);
+			Actor->SetActorLabel(Label);
+			return Actor;
+		}
+
+		void AddReplicatedProperty(const UObject* Object, FName RootPropertyName)
+		{
+			AddReplicatedProperties(Object, { FConcertPropertyChain::CreateFromPath(*Object->GetClass(), { RootPropertyName }) });
+		}
+		void AddReplicatedProperties(const UObject* Object, const TArray<TOptional<FConcertPropertyChain>>& Properties)
+		{
+			const bool bHasInvalid = Algo::AnyOf(Properties, [](const TOptional<FConcertPropertyChain>& Property){ return !Property.IsSet(); });
+			if (bHasInvalid)
+			{
+				UE_LOG(LogConcert, Error, TEXT("Did not resolve property"));
+				return;
+			}
+
+			FConcertReplicatedObjectInfo& ObjectInfo = ReplicationMap.ReplicatedObjects.FindOrAdd(
+				Object, { .ClassPath = Object->GetClass() }
+				);
+			Algo::Transform(Properties, ObjectInfo.PropertySelection.ReplicatedProperties,
+				[](const TOptional<FConcertPropertyChain>& Property){ return *Property; }
+				);
+		}
+	};
+
+	static void TestReplicatesOnlyProperty(const FConcertObjectReplicationMap& Map, const UObject* Object, FName PropertyName, FAutomationTestBase& Test)
+	{
+		const FConcertReplicatedObjectInfo* ActorData = Map.ReplicatedObjects.Find(Object);
+		if (ActorData)
+		{
+			Test.TestEqual(FString::Printf(TEXT("Class (%s)"), *Object->GetPathName()),
+				ActorData->ClassPath, FSoftClassPath(Object->GetClass())
+				);
+			const TSet<FConcertPropertyChain>& Properties = ActorData->PropertySelection.ReplicatedProperties;
+			if (Properties.Num() != 1)
+			{
+				Test.AddError(TEXT("Expected exactly 1 property"));
+			}
+			else
+			{
+				const FConcertPropertyChain& Property = *Properties.CreateConstIterator();
+				Test.TestEqual(TEXT("Root name"), Property.GetRootProperty(), PropertyName);
+				Test.TestEqual(TEXT("Path length"), Property.GetPathToProperty().Num(), 1);
+			}
+		}
+		else
+		{
+			Test.AddError(FString::Printf(TEXT("Property %s is not replicated for %s"), *PropertyName.ToString(), *Object->GetPathName()));
+		}
+	}
+}
\ No newline at end of file
