Added Concert utility for resetting object state to match another object

The new ResetObjectPropertiesToTemplateValues function shares its implementation with the older ResetObjectPropertiesToArchetypeValues function, but gives the caller much more control over what data is copied, and from which object.

#rnx

diff --git a/Engine/Plugins/Developer/Concert/ConcertSync/ConcertSyncCore/Source/ConcertSyncCore/Private/ConcertSyncArchives.cpp b/Engine/Plugins/Developer/Concert/ConcertSync/ConcertSyncCore/Source/ConcertSyncCore/Private/ConcertSyncArchives.cpp
--- a/Engine/Plugins/Developer/Concert/ConcertSync/ConcertSyncCore/Source/ConcertSyncCore/Private/ConcertSyncArchives.cpp
+++ b/Engine/Plugins/Developer/Concert/ConcertSync/ConcertSyncCore/Source/ConcertSyncCore/Private/ConcertSyncArchives.cpp
@@ -60,34 +60,69 @@ void GatherDefaultSubobjectPaths(const UObject* Obj, TSet<FSoftObjectPath>& OutS
 
 void ResetObjectPropertiesToArchetypeValues(UObject* Object, const bool InIncludeEditorOnlyData)
 {
-	class FArchetypePropertyWriter : public FObjectWriter
+	if (const UObject* ObjectArchetype = Object->GetArchetype())
+	{
+		FResetValueOptions Options;
+		Options.bIncludeEditorOnlyData = InIncludeEditorOnlyData;
+		Options.bIsDuplicate = false;
+		Options.ObjectsToSkip.Add(Object);
+		Options.ObjectsToSkip.Add(ObjectArchetype);
+		GatherDefaultSubobjectPaths(Object, Options.ObjectsToSkip);
+		GatherDefaultSubobjectPaths(ObjectArchetype, Options.ObjectsToSkip);
+
+		ResetObjectPropertiesToTemplateValues(Object, ObjectArchetype, Options);
+	}
+}
+
+void ResetObjectPropertiesToTemplateValues(UObject* Object, const UObject* Template, const FResetValueOptions& Options)
+{
+	class FTemplatePropertyWriter final : public FObjectWriter
 	{
 	public:
-		FArchetypePropertyWriter(const UObject* Obj, TArray<uint8>& OutBytes, TSet<FSoftObjectPath>& OutObjectsToSkip, const bool InIncludeEditorOnlyData)
+		FTemplatePropertyWriter(const UObject* Obj, TArray<uint8>& OutBytes, const FResetValueOptions& InOptions)
 			: FObjectWriter(OutBytes)
 		{
+			Options = &InOptions;
+
 			ArIgnoreClassRef = true;
 			ArIgnoreArchetypeRef = true;
 			ArNoDelta = true;
+			ArIsSaveGame = Options->bSaveGameOnly;
+			ArPortFlags |= (Options->bIsDuplicate ? PPF_Duplicate : PPF_None);
+			ArPropertySkipFlags |= (Options->bIsDuplicate ? CPF_DuplicateTransient : CPF_None);
 
-			SetFilterEditorOnly(!InIncludeEditorOnlyData);
-
-			GatherDefaultSubobjectPaths(Obj, OutObjectsToSkip);
+			SetFilterEditorOnly(!Options->bIncludeEditorOnlyData);
 
 			Obj->SerializeScriptProperties(*this);
 		}
+
+		virtual bool ShouldSkipProperty(const FProperty* InProperty) const override
+		{
+			return FObjectWriter::ShouldSkipProperty(InProperty)
+				|| InProperty->HasAnyPropertyFlags(ArPropertySkipFlags)
+				|| (Options->ShouldSkipProperty && Options->ShouldSkipProperty(*this, InProperty));
+		}
+
+	private:
+		const FResetValueOptions* Options = nullptr;
+		EPropertyFlags ArPropertySkipFlags = CPF_None;
 	};
 
-	class FArchetypePropertyReader : public FObjectReader
+	class FTemplatePropertyReader final : public FObjectReader
 	{
 	public:
-		FArchetypePropertyReader(UObject* Obj, const TArray<uint8>& InBytes, const TSet<FSoftObjectPath>& InObjectsToSkip, const bool InIncludeEditorOnlyData)
+		FTemplatePropertyReader(UObject* Obj, const TArray<uint8>& InBytes, const FResetValueOptions& InOptions)
 			: FObjectReader(InBytes)
 		{
+			Options = &InOptions;
+
 			ArIgnoreClassRef = true;
 			ArIgnoreArchetypeRef = true;
+			ArIsSaveGame = Options->bSaveGameOnly;
+			ArPortFlags |= (Options->bIsDuplicate ? PPF_Duplicate : PPF_None);
+			ArPropertySkipFlags |= (Options->bIsDuplicate ? CPF_DuplicateTransient : CPF_None);
 
-			SetFilterEditorOnly(!InIncludeEditorOnlyData);
+			SetFilterEditorOnly(!Options->bIncludeEditorOnlyData);
 
 #if USE_STABLE_LOCALIZATION_KEYS
 			if (GIsEditor && !(ArPortFlags & (PPF_DuplicateVerbatim | PPF_DuplicateForPIE)))
@@ -96,18 +131,22 @@ void ResetObjectPropertiesToArchetypeValues(UObject* Object, const bool InInclud
 			}
 #endif // USE_STABLE_LOCALIZATION_KEYS
 
-			ObjectsToSkip = InObjectsToSkip;
-			GatherDefaultSubobjectPaths(Obj, ObjectsToSkip);
-
 			Obj->SerializeScriptProperties(*this);
 		}
 
+		virtual bool ShouldSkipProperty(const FProperty* InProperty) const override
+		{
+			return FObjectReader::ShouldSkipProperty(InProperty)
+				|| InProperty->HasAnyPropertyFlags(ArPropertySkipFlags)
+				|| (Options->ShouldSkipProperty && Options->ShouldSkipProperty(*this, InProperty));
+		}
+
 		virtual FArchive& operator<<(UObject*& Value) override
 		{
 			UObject* Tmp = nullptr;
 			FObjectReader::operator<<(Tmp);
 
-			if (CanOverwriteObject(Value, Tmp))
+			if (Value != Tmp && CanOverwriteObject(Value, Tmp))
 			{
 				Value = Tmp;
 			}
@@ -120,7 +159,7 @@ void ResetObjectPropertiesToArchetypeValues(UObject* Object, const bool InInclud
 			FObjectPtr Tmp = nullptr;
 			FObjectReader::operator<<(Tmp);
 
-			if (CanOverwriteObject(Value.Get(), Tmp.Get()))
+			if (Value != Tmp && CanOverwriteObject(Value.Get(), Tmp.Get()))
 			{
 				Value = Tmp;
 			}
@@ -133,7 +172,7 @@ void ResetObjectPropertiesToArchetypeValues(UObject* Object, const bool InInclud
 			FLazyObjectPtr Tmp;
 			FObjectReader::operator<<(Tmp);
 
-			if (CanOverwriteObject(Value.Get(), Tmp.Get()))
+			if (Value != Tmp && CanOverwriteObject(Value.Get(), Tmp.Get()))
 			{
 				Value = Tmp;
 			}
@@ -146,7 +185,7 @@ void ResetObjectPropertiesToArchetypeValues(UObject* Object, const bool InInclud
 			FSoftObjectPtr Tmp;
 			FObjectReader::operator<<(Tmp);
 
-			if (CanOverwriteObject(Value.ToSoftObjectPath(), Tmp.ToSoftObjectPath()))
+			if (Value != Tmp && CanOverwriteObject(Value.ToSoftObjectPath(), Tmp.ToSoftObjectPath()))
 			{
 				Value = Tmp;
 			}
@@ -159,7 +198,7 @@ void ResetObjectPropertiesToArchetypeValues(UObject* Object, const bool InInclud
 			FSoftObjectPath Tmp;
 			FObjectReader::operator<<(Tmp);
 
-			if (CanOverwriteObject(Value, Tmp))
+			if (Value != Tmp && CanOverwriteObject(Value, Tmp))
 			{
 				Value = Tmp;
 			}
@@ -172,7 +211,7 @@ void ResetObjectPropertiesToArchetypeValues(UObject* Object, const bool InInclud
 			FWeakObjectPtr Tmp;
 			FObjectReader::operator<<(Tmp);
 
-			if (CanOverwriteObject(Value.Get(), Tmp.Get()))
+			if (Value != Tmp && CanOverwriteObject(Value.Get(), Tmp.Get()))
 			{
 				Value = Tmp;
 			}
@@ -183,19 +222,18 @@ void ResetObjectPropertiesToArchetypeValues(UObject* Object, const bool InInclud
 	private:
 		bool CanOverwriteObject(const FSoftObjectPath& CurObj, const FSoftObjectPath& NewObj) const
 		{
-			return !ObjectsToSkip.Contains(CurObj) && !ObjectsToSkip.Contains(NewObj);
+			return !Options->ObjectsToSkip.Contains(CurObj) 
+				&& !Options->ObjectsToSkip.Contains(NewObj)
+				&& !(Options->ShouldSkipObjectReference && Options->ShouldSkipObjectReference(*this, CurObj, NewObj));
 		}
 
-		TSet<FSoftObjectPath> ObjectsToSkip;
+		const FResetValueOptions* Options = nullptr;
+		EPropertyFlags ArPropertySkipFlags = CPF_None;
 	};
 
-	if (const UObject* ObjectArchetype = Object->GetArchetype())
-	{
-		TArray<uint8> ArchetypeData;
-		TSet<FSoftObjectPath> ArchetypeObjectsToSkip;
-		FArchetypePropertyWriter(ObjectArchetype, ArchetypeData, ArchetypeObjectsToSkip, InIncludeEditorOnlyData);
-		FArchetypePropertyReader(Object, ArchetypeData, ArchetypeObjectsToSkip, InIncludeEditorOnlyData);
-	}
+	TArray<uint8> TemplateData;
+	FTemplatePropertyWriter Writer(Template, TemplateData, Options);
+	FTemplatePropertyReader Reader(Object, TemplateData, Options);
 }
 
 } // namespace ConcertSyncUtil
diff --git a/Engine/Plugins/Developer/Concert/ConcertSync/ConcertSyncCore/Source/ConcertSyncCore/Public/ConcertSyncArchives.h b/Engine/Plugins/Developer/Concert/ConcertSync/ConcertSyncCore/Source/ConcertSyncCore/Public/ConcertSyncArchives.h
--- a/Engine/Plugins/Developer/Concert/ConcertSync/ConcertSyncCore/Source/ConcertSyncCore/Public/ConcertSyncArchives.h
+++ b/Engine/Plugins/Developer/Concert/ConcertSync/ConcertSyncCore/Source/ConcertSyncCore/Public/ConcertSyncArchives.h
@@ -23,6 +23,47 @@ CONCERTSYNCCORE_API bool CanExportProperty(const FProperty* Property, const bool
 
 CONCERTSYNCCORE_API void ResetObjectPropertiesToArchetypeValues(UObject* Object, const bool InIncludeEditorOnlyData);
 
+struct FResetValueOptions
+{
+	/**
+	 * True if we should also reset editor-only properties.
+	 */
+	bool bIncludeEditorOnlyData = true;
+
+	/**
+	 * True when Object is considered a duplicate of Template (eg, two instances with the same archetype, where we want to make Object the same as Template).
+	 * @note If true this sets PPF_Duplicate on the archive when copying data, and automatically skips CPF_DuplicateTransient properties.
+	 */
+	bool bIsDuplicate = false;
+
+	/**
+	 * True to only consider things marked CPF_SaveGame.
+	 * @note If true this sets the ArIsSaveGame flag on the archive when copying the data.
+	 */
+	bool bSaveGameOnly = false;
+
+	/**
+	 * Set of objects that should have their references skipped when resetting values.
+	 * @note Objects skipped via this set are not passed to ShouldSkipObjectReference.
+	 */
+	TSet<FSoftObjectPath> ObjectsToSkip;
+
+	/**
+	 * Filter used to decide whether the given property should be skipped when resetting values.
+	 * @note Properties skipped internally by the archive will not be passed to this function.
+	 */
+	TFunction<bool(const FArchive& Ar, const FProperty* Property)> ShouldSkipProperty;
+
+	/**
+	 * Filter used to decide whether the given property reference should be skipped when resetting values.
+	 * @param OldValue The current object reference from Object.
+	 * @param NewValue The proposed new object reference from Template.
+	 */
+	TFunction<bool(const FArchive& Ar, const FSoftObjectPath& OldValue, const FSoftObjectPath& NewValue)> ShouldSkipObjectReference;
+};
+
+CONCERTSYNCCORE_API void ResetObjectPropertiesToTemplateValues(UObject* Object, const UObject* Template, const FResetValueOptions& Options);
+
 CONCERTSYNCCORE_API const FSoftObjectPath& GetSkipObjectPath();
 }
 
