Fixed Concert object resetting being too agressive

* We no longer reset instanced object properties, as doing so can point an instance back to its CDO and break GC assumptions.
* This includes instanced object properties inside containers (like arrays).
* We no longer reset properties Concert isn't allowed to edit.
* This avoids resetting bookkeeping data (eg, an actor or entity component list).
* We now consider subobjects of a DSO to also be part of the DSO, and skip them when resetting.

diff --git a/Engine/Plugins/Developer/Concert/ConcertSync/ConcertSyncCore/Source/ConcertSyncCore/Private/ConcertSyncArchives.cpp b/Engine/Plugins/Developer/Concert/ConcertSync/ConcertSyncCore/Source/ConcertSyncCore/Private/ConcertSyncArchives.cpp
--- a/Engine/Plugins/Developer/Concert/ConcertSync/ConcertSyncCore/Source/ConcertSyncCore/Private/ConcertSyncArchives.cpp
+++ b/Engine/Plugins/Developer/Concert/ConcertSync/ConcertSyncCore/Source/ConcertSyncCore/Private/ConcertSyncArchives.cpp
@@ -5,6 +5,7 @@
 #include "ConcertVersion.h"
 
 #include "Misc/PackageName.h"
+#include "Serialization/ArchiveSerializedPropertyChain.h"
 #include "Serialization/ObjectReader.h"
 #include "Serialization/ObjectWriter.h"
 #include "UObject/Package.h"
@@ -47,15 +48,47 @@ bool CanExportProperty(const FProperty* Property, const bool InIncludeEditorOnly
 		&& (!PropertyTypeIsInList(SyncConfig->ExcludedPropertyTypes));
 }
 
-void GatherDefaultSubobjectPaths(const UObject* Obj, TSet<FSoftObjectPath>& OutSubobjects)
+void GatherDefaultSubobjectPaths(const UObject* Obj, TSet<FSoftObjectPath>& OutSubobjects, const bool bOuterIsDSO = false)
 {
-	ForEachObjectWithOuter(Obj, [&OutSubobjects](UObject* InnerObj)
+	ForEachObjectWithOuter(Obj, [bOuterIsDSO, &OutSubobjects](UObject* InnerObj)
 	{
-		if (InnerObj->HasAnyFlags(RF_DefaultSubObject) || InnerObj->IsDefaultSubobject())
+		if (bOuterIsDSO || InnerObj->HasAnyFlags(RF_DefaultSubObject) || InnerObj->IsDefaultSubobject())
 		{
 			OutSubobjects.Emplace(InnerObj);
+			GatherDefaultSubobjectPaths(InnerObj, OutSubobjects, /*bOuterIsDSO*/true);
 		}
-	});
+	}, /*bIncludeNestedObjects*/false);
+}
+
+bool ShouldSkipPropertyDuringReset(const FArchive& Ar, const FProperty* Property, const EPropertyFlags PropertySkipFlags, const FResetValueOptions& Options)
+{
+	if (Property->HasAnyPropertyFlags(PropertySkipFlags))
+	{
+		return true;
+	}
+
+	// Skip properties that Concert can't edit, as resetting them may destroy bookkeeping data
+	{
+		// Note: We do this using the root property being serialized (rather than the given property), 
+		//       as Concert filters currently only include top-level properties
+		const FProperty* RootProperty = Property;
+		if (const FArchiveSerializedPropertyChain* PropertyChain = Ar.GetSerializedPropertyChain();
+			PropertyChain && PropertyChain->GetNumProperties() > 0)
+		{
+			RootProperty = PropertyChain->GetPropertyFromRoot(0);
+		}
+		if (!CanExportProperty(RootProperty, Options.bIncludeEditorOnlyData))
+		{
+			return true;
+		}
+	}
+
+	if (Options.ShouldSkipProperty && Options.ShouldSkipProperty(Ar, Property))
+	{
+		return true;
+	}
+
+	return false;
 }
 
 void ResetObjectPropertiesToArchetypeValues(UObject* Object, const bool InIncludeEditorOnlyData)
@@ -76,6 +109,8 @@ void ResetObjectPropertiesToArchetypeValues(UObject* Object, const bool InInclud
 
 void ResetObjectPropertiesToTemplateValues(UObject* Object, const UObject* Template, const FResetValueOptions& Options)
 {
+	static constexpr EPropertyFlags InstancedObjectPropertyFlags = CPF_InstancedReference | CPF_PersistentInstance;
+
 	class FTemplatePropertyWriter final : public FObjectWriter
 	{
 	public:
@@ -90,6 +125,7 @@ void ResetObjectPropertiesToTemplateValues(UObject* Object, const UObject* Templ
 			ArIsSaveGame = Options->bSaveGameOnly;
 			ArPortFlags |= (Options->bIsDuplicate ? PPF_Duplicate : PPF_None);
 			ArPropertySkipFlags |= (Options->bIsDuplicate ? CPF_DuplicateTransient : CPF_None);
+			ArPropertySkipFlags |= InstancedObjectPropertyFlags; // We can't handle resetting instanced objects
 
 			SetFilterEditorOnly(!Options->bIncludeEditorOnlyData);
 
@@ -99,8 +135,7 @@ void ResetObjectPropertiesToTemplateValues(UObject* Object, const UObject* Templ
 		virtual bool ShouldSkipProperty(const FProperty* InProperty) const override
 		{
 			return FObjectWriter::ShouldSkipProperty(InProperty)
-				|| InProperty->HasAnyPropertyFlags(ArPropertySkipFlags)
-				|| (Options->ShouldSkipProperty && Options->ShouldSkipProperty(*this, InProperty));
+				|| ShouldSkipPropertyDuringReset(*this, InProperty, ArPropertySkipFlags, *Options);
 		}
 
 	private:
@@ -121,6 +156,7 @@ void ResetObjectPropertiesToTemplateValues(UObject* Object, const UObject* Templ
 			ArIsSaveGame = Options->bSaveGameOnly;
 			ArPortFlags |= (Options->bIsDuplicate ? PPF_Duplicate : PPF_None);
 			ArPropertySkipFlags |= (Options->bIsDuplicate ? CPF_DuplicateTransient : CPF_None);
+			ArPropertySkipFlags |= InstancedObjectPropertyFlags; // We can't handle resetting instanced objects
 
 			SetFilterEditorOnly(!Options->bIncludeEditorOnlyData);
 
@@ -137,8 +173,7 @@ void ResetObjectPropertiesToTemplateValues(UObject* Object, const UObject* Templ
 		virtual bool ShouldSkipProperty(const FProperty* InProperty) const override
 		{
 			return FObjectReader::ShouldSkipProperty(InProperty)
-				|| InProperty->HasAnyPropertyFlags(ArPropertySkipFlags)
-				|| (Options->ShouldSkipProperty && Options->ShouldSkipProperty(*this, InProperty));
+				|| ShouldSkipPropertyDuringReset(*this, InProperty, ArPropertySkipFlags, *Options);
 		}
 
 		virtual FArchive& operator<<(UObject*& Value) override
@@ -222,7 +257,13 @@ void ResetObjectPropertiesToTemplateValues(UObject* Object, const UObject* Templ
 	private:
 		bool CanOverwriteObject(const FSoftObjectPath& CurObj, const FSoftObjectPath& NewObj) const
 		{
-			return !Options->ObjectsToSkip.Contains(CurObj) 
+			// Archive serialization doesn't catch all CPF_InstancedReference or CPF_PersistentInstance properties via ShouldSkipProperty, 
+			// as it's not called in some cases, eg, a TArray<UObject*> will have the FArrayProperty marked as CPF_ContainsInstancedReference 
+			// but won't also test the inner FObjectProperty before serializing its value; instead we query the head of the property 
+			// serialization stack and verify it isn't an instanced object property before allowing an object replacement to happen
+			const FProperty* Prop = GetSerializedProperty();
+			return !(Prop && Prop->HasAnyPropertyFlags(InstancedObjectPropertyFlags))
+				&& !Options->ObjectsToSkip.Contains(CurObj) 
 				&& !Options->ObjectsToSkip.Contains(NewObj)
 				&& !(Options->ShouldSkipObjectReference && Options->ShouldSkipObjectReference(*this, CurObj, NewObj));
 		}
