Replaced some FSoftObjectPath::GetSubPathString with GetSubPathUtf8String.

diff --git a/Engine/Plugins/Developer/Concert/ConcertSync/ConcertSyncCore/Source/ConcertSyncCore/Private/Misc/ObjectUtils.cpp b/Engine/Plugins/Developer/Concert/ConcertSync/ConcertSyncCore/Source/ConcertSyncCore/Private/Misc/ObjectUtils.cpp
--- a/Engine/Plugins/Developer/Concert/ConcertSync/ConcertSyncCore/Source/ConcertSyncCore/Private/Misc/ObjectUtils.cpp
+++ b/Engine/Plugins/Developer/Concert/ConcertSync/ConcertSyncCore/Source/ConcertSyncCore/Private/Misc/ObjectUtils.cpp
@@ -11,10 +11,10 @@ namespace UE::ConcertSyncCore
 	{
 		// Example of an actor called floor
 		// SoftObjectPath = { AssetPath = {PackageName = "/Game/Maps/SyncBoxLevel", AssetName = "SyncBoxLevel"}, SubPathString = "PersistentLevel.Floor" } }
-		const FString& SubPathString = SoftObjectPath.GetSubPathString();
+		const FUtf8String& SubPathString = SoftObjectPath.GetSubPathUtf8String();
 
 		constexpr int32 PersistentLevelStringLength = 16; // "PersistentLevel." has 16 characters
-		const bool bIsWorldObject = SubPathString.Contains(TEXT("PersistentLevel."), ESearchCase::CaseSensitive);
+		const bool bIsWorldObject = SubPathString.Contains("PersistentLevel.", ESearchCase::CaseSensitive);
 		if (!bIsWorldObject)
 		{
 			// Not a path to a world object
@@ -23,7 +23,7 @@ namespace UE::ConcertSyncCore
 
 		// Start search after the . behind PersistentLevel
 		const int32 StartSearch = PersistentLevelStringLength + 1;
-		const int32 IndexOfDotAfterActorName = SubPathString.Find(TEXT("."), ESearchCase::CaseSensitive, ESearchDir::FromStart, StartSearch);
+		const int32 IndexOfDotAfterActorName = SubPathString.Find(".", ESearchCase::CaseSensitive, ESearchDir::FromStart, StartSearch);
 		return IndexOfDotAfterActorName == INDEX_NONE;
 	}
 	
@@ -31,10 +31,10 @@ namespace UE::ConcertSyncCore
 	{
 		// Example of an actor called floor
 		// SoftObjectPath = { AssetPath = {PackageName = "/Game/Maps/SyncBoxLevel", AssetName = "SyncBoxLevel"}, SubPathString = "PersistentLevel.Floor" } }
-		const FString& SubPathString = SoftObjectPath.GetSubPathString();
+		const FUtf8String& SubPathString = SoftObjectPath.GetSubPathUtf8String();
 
 		constexpr int32 PersistentLevelStringLength = 16; // "PersistentLevel." has 16 characters
-		const bool bIsWorldObject = SubPathString.Contains(TEXT("PersistentLevel."), ESearchCase::CaseSensitive);
+		const bool bIsWorldObject = SubPathString.Contains("PersistentLevel.", ESearchCase::CaseSensitive);
 		if (!bIsWorldObject)
 		{
 			// Not a path to a world object
@@ -43,7 +43,7 @@ namespace UE::ConcertSyncCore
 
 		// Start search after the . behind PersistentLevel
 		const int32 StartSearch = PersistentLevelStringLength + 1;
-		const int32 IndexOfDotAfterActorName = SubPathString.Find(TEXT("."), ESearchCase::CaseSensitive, ESearchDir::FromStart, StartSearch);
+		const int32 IndexOfDotAfterActorName = SubPathString.Find(".", ESearchCase::CaseSensitive, ESearchDir::FromStart, StartSearch);
 		if (IndexOfDotAfterActorName == INDEX_NONE)
 		{
 			// SoftObjectPath points to an actor
@@ -51,16 +51,16 @@ namespace UE::ConcertSyncCore
 		}
 
 		const int32 NumToChopOffRight = SubPathString.Len() - IndexOfDotAfterActorName;
-		const FString NewSubstring = SubPathString.LeftChop(NumToChopOffRight);
+		const FUtf8String NewSubstring = SubPathString.LeftChop(NumToChopOffRight);
 		const FSoftObjectPath PathToOwningActor(SoftObjectPath.GetAssetPath(), NewSubstring);
 		return PathToOwningActor;
 	}
 	
-	FString ExtractObjectNameFromPath(const FSoftObjectPath& Object)
+	FUtf8String ExtractObjectNameFromPath(const FSoftObjectPath& Object)
 	{
 		// Subpath looks like this PersistentLevel.Actor.Component
-		const FString& Subpath = Object.GetSubPathString();
-		const int32 LastDotIndex = Subpath.Find(TEXT("."), ESearchCase::CaseSensitive, ESearchDir::FromEnd);
+		const FUtf8String& Subpath = Object.GetSubPathUtf8String();
+		const int32 LastDotIndex = Subpath.Find(".", ESearchCase::CaseSensitive, ESearchDir::FromEnd);
 		if (LastDotIndex == INDEX_NONE)
 		{
 			return {};
@@ -77,16 +77,16 @@ namespace UE::ConcertSyncCore
 		
 		// Example of an actor called floor
 		// SoftObjectPath = { AssetPath = {PackageName = "/Game/Maps/SyncBoxLevel", AssetName = "SyncBoxLevel"}, SubPathString = "PersistentLevel.Floor" } }
-		const FString& OldSubPathString = OldPath.GetSubPathString();
-		if (!OldSubPathString.Contains(TEXT("PersistentLevel."), ESearchCase::CaseSensitive))
+		const FUtf8String& OldSubPathString = OldPath.GetSubPathUtf8String();
+		if (!OldSubPathString.Contains("PersistentLevel.", ESearchCase::CaseSensitive))
 		{
 			return {};
 		}
-		const FString& NewSubPathString = NewActor.GetSubPathString();
+		const FUtf8String& NewSubPathString = NewActor.GetSubPathUtf8String();
 		
 		constexpr int32 PersistentLevelStringLength = 16; // "PersistentLevel." has 16 characters
 		constexpr int32 FirstActorCharIndex = PersistentLevelStringLength + 1;
-		const int32 IndexOfDotAfterActorName_OldSubPath = OldSubPathString.Find(TEXT("."), ESearchCase::CaseSensitive, ESearchDir::FromStart, FirstActorCharIndex);
+		const int32 IndexOfDotAfterActorName_OldSubPath = OldSubPathString.Find(".", ESearchCase::CaseSensitive, ESearchDir::FromStart, FirstActorCharIndex);
 		
 		const bool bOldIsOnlyActor = !OldSubPathString.IsValidIndex(IndexOfDotAfterActorName_OldSubPath);
 		if (bOldIsOnlyActor)
@@ -94,8 +94,8 @@ namespace UE::ConcertSyncCore
 			return NewActor;
 		}
 
-		const FString ReplacedSubPathString = NewSubPathString
-			+ TEXT(".")
+		const FUtf8String ReplacedSubPathString = NewSubPathString
+			+ "."
 			+ OldSubPathString.RightChop(IndexOfDotAfterActorName_OldSubPath + 1);
 		return FSoftObjectPath(NewActor.GetAssetPath(), ReplacedSubPathString);
 	}
diff --git a/Engine/Plugins/Developer/Concert/ConcertSync/ConcertSyncCore/Source/ConcertSyncCore/Public/Misc/ObjectPathUtils.h b/Engine/Plugins/Developer/Concert/ConcertSync/ConcertSyncCore/Source/ConcertSyncCore/Public/Misc/ObjectPathUtils.h
--- a/Engine/Plugins/Developer/Concert/ConcertSync/ConcertSyncCore/Source/ConcertSyncCore/Public/Misc/ObjectPathUtils.h
+++ b/Engine/Plugins/Developer/Concert/ConcertSync/ConcertSyncCore/Source/ConcertSyncCore/Public/Misc/ObjectPathUtils.h
@@ -1,8 +1,8 @@
-﻿// Copyright Epic Games, Inc. All Rights Reserved.
+// Copyright Epic Games, Inc. All Rights Reserved.
 
 #pragma once
 
-#include "Containers/UnrealString.h"
+#include "Containers/Utf8String.h"
 #include "Misc/Optional.h"
 #include "UObject/SoftObjectPath.h"
 
@@ -32,8 +32,8 @@ namespace UE::ConcertSyncCore
 			return {};
 		}
 		
-		FString SubPathString = ObjectPath.GetSubPathString();
-		const int32 CurrentIndex = SubPathString.Find(TEXT("."), ESearchCase::CaseSensitive, ESearchDir::FromEnd);
+		FUtf8String SubPathString = ObjectPath.GetSubPathUtf8String();
+		const int32 CurrentIndex = SubPathString.Find(".", ESearchCase::CaseSensitive, ESearchDir::FromEnd);
 		if (SubPathString.IsValidIndex(CurrentIndex))
 		{
 			SubPathString.LeftInline(CurrentIndex);
diff --git a/Engine/Plugins/Developer/Concert/ConcertSync/ConcertSyncCore/Source/ConcertSyncCore/Public/Misc/ObjectUtils.h b/Engine/Plugins/Developer/Concert/ConcertSync/ConcertSyncCore/Source/ConcertSyncCore/Public/Misc/ObjectUtils.h
--- a/Engine/Plugins/Developer/Concert/ConcertSync/ConcertSyncCore/Source/ConcertSyncCore/Public/Misc/ObjectUtils.h
+++ b/Engine/Plugins/Developer/Concert/ConcertSync/ConcertSyncCore/Source/ConcertSyncCore/Public/Misc/ObjectUtils.h
@@ -16,7 +16,7 @@ namespace UE::ConcertSyncCore
 	FORCEINLINE TOptional<FSoftObjectPath> GetActorPathIn(const FSoftObjectPath& Path) { return IsActor(Path) ? TOptional(Path) : GetActorOf(Path); }
 	
 	/** @return Gets the last object name in the subpath. */
-	CONCERTSYNCCORE_API FString ExtractObjectNameFromPath(const FSoftObjectPath& Object);
+	CONCERTSYNCCORE_API FUtf8String ExtractObjectNameFromPath(const FSoftObjectPath& Object);
 
 	/**
 	 * Replaces the package path and actor name from NewActor witht that of NewActor.
