Ran UnrealCodeFixup script on code to fixup dllstorage

diff --git a/Engine/Plugins/Developer/Concert/ConcertMain/Source/Concert/Public/ConcertSession.h b/Engine/Plugins/Developer/Concert/ConcertMain/Source/Concert/Public/ConcertSession.h
--- a/Engine/Plugins/Developer/Concert/ConcertMain/Source/Concert/Public/ConcertSession.h
+++ b/Engine/Plugins/Developer/Concert/ConcertMain/Source/Concert/Public/ConcertSession.h
@@ -6,18 +6,20 @@
 #include "IConcertSession.h"
 #include "ConcertMessages.h"
 
+#define UE_API CONCERT_API
+
 /**
  * Common implementation for Concert Client and Server sessions (@see IConcertSession for the API description).
  * @note	This doesn't inherit from any session interface, but does implement some of their API with a "Common" prefix on the function names.
  *			Interface implementations can also inherit from this common impl and then call the "Common" functions from the interface overrides.
  */
-class CONCERT_API FConcertSessionCommonImpl
+class FConcertSessionCommonImpl
 {
 public:
-	explicit FConcertSessionCommonImpl(const FConcertSessionInfo& InSessionInfo);
+	UE_API explicit FConcertSessionCommonImpl(const FConcertSessionInfo& InSessionInfo);
 
-	void CommonStartup();
-	void CommonShutdown();
+	UE_API void CommonStartup();
+	UE_API void CommonShutdown();
 
 	const FGuid& CommonGetId() const
 	{
@@ -39,26 +41,26 @@ public:
 		return SessionInfo;
 	}
 
-	TArray<FGuid> CommonGetSessionClientEndpointIds() const;
-	TArray<FConcertSessionClientInfo> CommonGetSessionClients() const;
-	bool CommonFindSessionClient(const FGuid& EndpointId, FConcertSessionClientInfo& OutSessionClientInfo) const;
+	UE_API TArray<FGuid> CommonGetSessionClientEndpointIds() const;
+	UE_API TArray<FConcertSessionClientInfo> CommonGetSessionClients() const;
+	UE_API bool CommonFindSessionClient(const FGuid& EndpointId, FConcertSessionClientInfo& OutSessionClientInfo) const;
 
-	FConcertScratchpadRef CommonGetScratchpad() const;
-	FConcertScratchpadPtr CommonGetClientScratchpad(const FGuid& ClientEndpointId) const;
+	UE_API FConcertScratchpadRef CommonGetScratchpad() const;
+	UE_API FConcertScratchpadPtr CommonGetClientScratchpad(const FGuid& ClientEndpointId) const;
 
-	FDelegateHandle CommonRegisterCustomEventHandler(const FName& EventMessageType, const TSharedRef<IConcertSessionCustomEventHandler>& Handler);
-	void CommonUnregisterCustomEventHandler(const FName& EventMessageType, const FDelegateHandle EventHandle);
-	void CommonUnregisterCustomEventHandler(const FName& EventMessageType, const void* EventHandler);
-	void CommonClearCustomEventHandler(const FName& EventMessageType);
-	void CommonHandleCustomEvent(const FConcertMessageContext& Context);
+	UE_API FDelegateHandle CommonRegisterCustomEventHandler(const FName& EventMessageType, const TSharedRef<IConcertSessionCustomEventHandler>& Handler);
+	UE_API void CommonUnregisterCustomEventHandler(const FName& EventMessageType, const FDelegateHandle EventHandle);
+	UE_API void CommonUnregisterCustomEventHandler(const FName& EventMessageType, const void* EventHandler);
+	UE_API void CommonClearCustomEventHandler(const FName& EventMessageType);
+	UE_API void CommonHandleCustomEvent(const FConcertMessageContext& Context);
 
-	void CommonRegisterCustomRequestHandler(const FName& RequestMessageType, const TSharedRef<IConcertSessionCustomRequestHandler>& Handler);
-	void CommonUnregisterCustomRequestHandler(const FName& RequestMessageType);
-	TFuture<FConcertSession_CustomResponse> CommonHandleCustomRequest(const FConcertMessageContext& Context);
+	UE_API void CommonRegisterCustomRequestHandler(const FName& RequestMessageType, const TSharedRef<IConcertSessionCustomRequestHandler>& Handler);
+	UE_API void CommonUnregisterCustomRequestHandler(const FName& RequestMessageType);
+	UE_API TFuture<FConcertSession_CustomResponse> CommonHandleCustomRequest(const FConcertMessageContext& Context);
 
-	static bool CommonBuildCustomEvent(const UScriptStruct* EventType, const void* EventData, const FGuid& SourceEndpointId, const TArray<FGuid>& DestinationEndpointIds, FConcertSession_CustomEvent& OutCustomEvent);
-	static bool CommonBuildCustomRequest(const UScriptStruct* RequestType, const void* RequestData, const FGuid& SourceEndpointId, const FGuid& DestinationEndpointId, FConcertSession_CustomRequest& OutCustomRequest);
-	static void CommonHandleCustomResponse(const FConcertSession_CustomResponse& Response, const TSharedRef<IConcertSessionCustomResponseHandler>& Handler);
+	static UE_API bool CommonBuildCustomEvent(const UScriptStruct* EventType, const void* EventData, const FGuid& SourceEndpointId, const TArray<FGuid>& DestinationEndpointIds, FConcertSession_CustomEvent& OutCustomEvent);
+	static UE_API bool CommonBuildCustomRequest(const UScriptStruct* RequestType, const void* RequestData, const FGuid& SourceEndpointId, const FGuid& DestinationEndpointId, FConcertSession_CustomRequest& OutCustomRequest);
+	static UE_API void CommonHandleCustomResponse(const FConcertSession_CustomResponse& Response, const TSharedRef<IConcertSessionCustomResponseHandler>& Handler);
 
 protected:
 	/** Information about this session */
@@ -81,3 +83,5 @@ protected:
 	/** Map of custom request handlers for this session */
 	TMap<FName, TSharedPtr<IConcertSessionCustomRequestHandler>> CustomRequestHandlers;
 };
+
+#undef UE_API
diff --git a/Engine/Plugins/Developer/Concert/ConcertMain/Source/Concert/Public/ConcertTransportEvents.h b/Engine/Plugins/Developer/Concert/ConcertMain/Source/Concert/Public/ConcertTransportEvents.h
--- a/Engine/Plugins/Developer/Concert/ConcertMain/Source/Concert/Public/ConcertTransportEvents.h
+++ b/Engine/Plugins/Developer/Concert/ConcertMain/Source/Concert/Public/ConcertTransportEvents.h
@@ -27,7 +27,7 @@ enum class EConcertLogMessageAction : uint8
 };
 
 USTRUCT()
-struct CONCERT_API FConcertLog
+struct FConcertLog
 {
 	GENERATED_BODY()
 
diff --git a/Engine/Plugins/Developer/Concert/ConcertMain/Source/ConcertClient/Public/ConcertClientSettings.h b/Engine/Plugins/Developer/Concert/ConcertMain/Source/ConcertClient/Public/ConcertClientSettings.h
--- a/Engine/Plugins/Developer/Concert/ConcertMain/Source/ConcertClient/Public/ConcertClientSettings.h
+++ b/Engine/Plugins/Developer/Concert/ConcertMain/Source/ConcertClient/Public/ConcertClientSettings.h
@@ -13,6 +13,8 @@
 
 #include "ConcertClientSettings.generated.h"
 
+#define UE_API CONCERTCLIENT_API
+
 USTRUCT()
 struct FConcertClientSettings
 {
@@ -120,12 +122,12 @@ enum class EConcertServerType
 	Slate
 };
 
-UCLASS(config=Engine)
-class CONCERTCLIENT_API UConcertClientConfig : public UObject
+UCLASS(MinimalAPI, config=Engine)
+class UConcertClientConfig : public UObject
 {
 	GENERATED_BODY()
 public:
-	UConcertClientConfig();
+	UE_API UConcertClientConfig();
 
 	/*
 	 * Mark this setting object as editor only.
@@ -237,3 +239,5 @@ public:
 	UPROPERTY(config, EditAnywhere, AdvancedDisplay, Category="Endpoint Settings", meta=(ShowOnlyInnerProperties))
 	FConcertEndpointSettings EndpointSettings;
 };
+
+#undef UE_API
diff --git a/Engine/Plugins/Developer/Concert/ConcertMain/Source/ConcertTransport/Public/IdentifierTable/ConcertIdentifierTable.h b/Engine/Plugins/Developer/Concert/ConcertMain/Source/ConcertTransport/Public/IdentifierTable/ConcertIdentifierTable.h
--- a/Engine/Plugins/Developer/Concert/ConcertMain/Source/ConcertTransport/Public/IdentifierTable/ConcertIdentifierTable.h
+++ b/Engine/Plugins/Developer/Concert/ConcertMain/Source/ConcertTransport/Public/IdentifierTable/ConcertIdentifierTable.h
@@ -5,6 +5,8 @@
 #include "CoreMinimal.h"
 #include "ConcertIdentifierTableData.h"
 
+#define UE_API CONCERTTRANSPORT_API
+
 /** Key functions that compare FName instances by only their display string (to perform a case-sensitive plain string comparison) */
 template <typename ValueType>
 struct TConcertIdentifierTable_CaseSensitivePlainNameKeyFuncs : BaseKeyFuncs<ValueType, FName, /*bInAllowDuplicateKeys*/false>
@@ -24,52 +26,52 @@ struct TConcertIdentifierTable_CaseSensitivePlainNameKeyFuncs : BaseKeyFuncs<Val
 };
 
 /** Cache of identifiers (currently names) that have been serialized locally and should be sent along with the serialized data */
-class CONCERTTRANSPORT_API FConcertLocalIdentifierTable
+class FConcertLocalIdentifierTable
 {
 public:
 	FConcertLocalIdentifierTable() = default;
-	explicit FConcertLocalIdentifierTable(const FConcertLocalIdentifierState& InState);
+	UE_API explicit FConcertLocalIdentifierTable(const FConcertLocalIdentifierState& InState);
 
 	/**
 	 * Constant used to force names to serialize as their FName table index value (non-portable!)
 	 */
-	static FConcertLocalIdentifierTable* ForceFNameTableIndex;
+	static UE_API FConcertLocalIdentifierTable* ForceFNameTableIndex;
 
 	/**
 	 * Does this table have any mappings in it?
 	 */
-	bool HasMappings() const;
+	UE_API bool HasMappings() const;
 
 	/**
 	 * Map the given name to its identifier index.
 	 */
-	int32 MapName(const FName& InName);
+	UE_API int32 MapName(const FName& InName);
 
 	/**
 	 * Unmap the given name from its identifier index (or NAME_None).
 	 * @return true if the name was unmapped, false if the name wasn't mapped (OutName will be set to NAME_None).
 	 */
-	bool UnmapName(const int32 InIndex, FName& OutName) const;
+	UE_API bool UnmapName(const int32 InIndex, FName& OutName) const;
 
 	/**
 	 * Is the given identifier index mapped?
 	 */
-	bool HasName(const int32 InIndex) const;
+	UE_API bool HasName(const int32 InIndex) const;
 
 	/**
 	 * Is the given name mapped?
 	 */
-	bool HasName(const FName& InName, int32* OutIndexPtr = nullptr) const;
+	UE_API bool HasName(const FName& InName, int32* OutIndexPtr = nullptr) const;
 
 	/**
 	 * Set the state of this identifier table.
 	 */
-	void SetState(const FConcertLocalIdentifierState& InState);
+	UE_API void SetState(const FConcertLocalIdentifierState& InState);
 
 	/**
 	 * Get the state of this identifier table.
 	 */
-	void GetState(FConcertLocalIdentifierState& OutState) const;
+	UE_API void GetState(FConcertLocalIdentifierState& OutState) const;
 
 private:
 	/** Array of locally mapped names */
@@ -77,3 +79,5 @@ private:
 	/** Map of known names to their index in MappedNames (name -> index) */
 	TMap<FName, int32, FDefaultSetAllocator, TConcertIdentifierTable_CaseSensitivePlainNameKeyFuncs<int32>> NameToMappedIndex;
 };
+
+#undef UE_API
diff --git a/Engine/Plugins/Developer/Concert/ConcertMain/Source/ConcertTransport/Public/IdentifierTable/ConcertTransportArchives.h b/Engine/Plugins/Developer/Concert/ConcertMain/Source/ConcertTransport/Public/IdentifierTable/ConcertTransportArchives.h
--- a/Engine/Plugins/Developer/Concert/ConcertMain/Source/ConcertTransport/Public/IdentifierTable/ConcertTransportArchives.h
+++ b/Engine/Plugins/Developer/Concert/ConcertMain/Source/ConcertTransport/Public/IdentifierTable/ConcertTransportArchives.h
@@ -6,20 +6,22 @@
 #include "Serialization/MemoryWriter.h"
 #include "Serialization/MemoryReader.h"
 
+#define UE_API CONCERTTRANSPORT_API
+
 class FConcertLocalIdentifierTable;
 
 /** Archive for writing identifiers (currently names) in a way that avoids duplication by caching them against their internal key, which can then be mapped over the network */
-class CONCERTTRANSPORT_API FConcertIdentifierWriter : public FMemoryWriter
+class FConcertIdentifierWriter : public FMemoryWriter
 {
 public:
-	FConcertIdentifierWriter(FConcertLocalIdentifierTable* InLocalIdentifierTable, TArray<uint8>& InBytes, bool bIsPersistent = false);
+	UE_API FConcertIdentifierWriter(FConcertLocalIdentifierTable* InLocalIdentifierTable, TArray<uint8>& InBytes, bool bIsPersistent = false);
 
 	using FMemoryWriter::operator<<; // For visibility of the overloads we don't override
 
 	//~ Begin FArchive Interface
-	virtual FArchive& operator<<(FName& Name) override;
-	virtual FArchive& operator<<(FSoftObjectPath& AssetPtr) override;
-	virtual FString GetArchiveName() const override;
+	UE_API virtual FArchive& operator<<(FName& Name) override;
+	UE_API virtual FArchive& operator<<(FSoftObjectPath& AssetPtr) override;
+	UE_API virtual FString GetArchiveName() const override;
 	//~ End FArchive Interface
 
 protected:
@@ -30,17 +32,17 @@ private:
 };
 
 /** Archive for reading identifiers (currently names) in a way that avoids duplication by caching them against their internal key, which can then be mapped over the network */
-class CONCERTTRANSPORT_API FConcertIdentifierReader : public FMemoryReader
+class FConcertIdentifierReader : public FMemoryReader
 {
 public:
-	FConcertIdentifierReader(const FConcertLocalIdentifierTable* InLocalIdentifierTable, const TArray<uint8>& InBytes, bool bIsPersistent = false);
+	UE_API FConcertIdentifierReader(const FConcertLocalIdentifierTable* InLocalIdentifierTable, const TArray<uint8>& InBytes, bool bIsPersistent = false);
 
 	using FMemoryReader::operator<<; // For visibility of the overloads we don't override
 
 	//~ Begin FArchive Interface
-	virtual FArchive& operator<<(FName& Name) override;
-	FArchive& operator<<(FSoftObjectPath& AssetPtr) override;
-	virtual FString GetArchiveName() const override;
+	UE_API virtual FArchive& operator<<(FName& Name) override;
+	UE_API FArchive& operator<<(FSoftObjectPath& AssetPtr) override;
+	UE_API virtual FString GetArchiveName() const override;
 	//~ End FArchive Interface
 
 protected:
@@ -51,19 +53,19 @@ private:
 };
 
 /** Archive for rewriting identifiers (currently names) so that they belong to a different identifier table */
-class CONCERTTRANSPORT_API FConcertIdentifierRewriter : public FMemoryArchive
+class FConcertIdentifierRewriter : public FMemoryArchive
 {
 public:
-	FConcertIdentifierRewriter(const FConcertLocalIdentifierTable* InLocalIdentifierTable, FConcertLocalIdentifierTable* InRewriteIdentifierTable, TArray<uint8>& InBytes, bool bIsPersistent = false);
+	UE_API FConcertIdentifierRewriter(const FConcertLocalIdentifierTable* InLocalIdentifierTable, FConcertLocalIdentifierTable* InRewriteIdentifierTable, TArray<uint8>& InBytes, bool bIsPersistent = false);
 
 	using FMemoryArchive::operator<<; // For visibility of the overloads we don't override
 
 	//~ Begin FArchive Interface
-	virtual FArchive& operator<<(FName& Name) override;
-	virtual FArchive& operator<<(FSoftObjectPath& AssetPtr) override;
-	virtual int64 TotalSize() override final;
-	virtual void Serialize(void* Data, int64 Num) override final;
-	virtual FString GetArchiveName() const override;
+	UE_API virtual FArchive& operator<<(FName& Name) override;
+	UE_API virtual FArchive& operator<<(FSoftObjectPath& AssetPtr) override;
+	UE_API virtual int64 TotalSize() override final;
+	UE_API virtual void Serialize(void* Data, int64 Num) override final;
+	UE_API virtual FString GetArchiveName() const override;
 	//~ End FArchive Interface
 
 protected:
@@ -102,3 +104,5 @@ private:
 	TArray<uint8>& Bytes;
 	TArray<uint8> ScratchBytes;
 };
+
+#undef UE_API
diff --git a/Engine/Plugins/Developer/Concert/ConcertMain/Source/ConcertTransport/Public/Scratchpad/ConcertScratchpad.h b/Engine/Plugins/Developer/Concert/ConcertMain/Source/ConcertTransport/Public/Scratchpad/ConcertScratchpad.h
--- a/Engine/Plugins/Developer/Concert/ConcertMain/Source/ConcertTransport/Public/Scratchpad/ConcertScratchpad.h
+++ b/Engine/Plugins/Developer/Concert/ConcertMain/Source/ConcertTransport/Public/Scratchpad/ConcertScratchpad.h
@@ -4,8 +4,10 @@
 
 #include "CoreMinimal.h"
 
+#define UE_API CONCERTTRANSPORT_API
+
 /** Scratchpad of id -> value pairs corresponding to a particular client */
-class CONCERTTRANSPORT_API FConcertScratchpad
+class FConcertScratchpad
 {
 private:
 	/** Scratchpad value interface */
@@ -43,7 +45,7 @@ public:
 	/**
 	 * Does the scratchpad have a value for the given key?
 	 */
-	bool HasValue(const FName& InId) const;
+	UE_API bool HasValue(const FName& InId) const;
 
 	/**
 	 * Set the scratchpad value associated with the given key.
@@ -109,15 +111,17 @@ private:
 	/**
 	 * Set the internal scratchpad value associated with the given key.
 	 */
-	void InternalSetValue(const FName& InId, IScratchpadValuePtr&& InValue);
+	UE_API void InternalSetValue(const FName& InId, IScratchpadValuePtr&& InValue);
 
 	/**
 	 * Get the internal scratchpad value associated with the given key (if any).
 	 */
-	IScratchpadValuePtr InternalGetValue(const FName& InId) const;
+	UE_API IScratchpadValuePtr InternalGetValue(const FName& InId) const;
 
 	/** Critical section protecting concurrent access to ScratchpadValues */
 	mutable FCriticalSection ScratchpadValuesCS;
 	/** Map of values (id -> value) */
 	TMap<FName, IScratchpadValuePtr> ScratchpadValues;
 };
+
+#undef UE_API
diff --git a/Engine/Plugins/Developer/Concert/ConcertMain/Source/ConcertTransport/Public/Trace/Private/ConcertScopedObjectTrace.h b/Engine/Plugins/Developer/Concert/ConcertMain/Source/ConcertTransport/Public/Trace/Private/ConcertScopedObjectTrace.h
--- a/Engine/Plugins/Developer/Concert/ConcertMain/Source/ConcertTransport/Public/Trace/Private/ConcertScopedObjectTrace.h
+++ b/Engine/Plugins/Developer/Concert/ConcertMain/Source/ConcertTransport/Public/Trace/Private/ConcertScopedObjectTrace.h
@@ -12,7 +12,7 @@
 
 namespace UE::ConcertTrace
 {
-	class CONCERTTRANSPORT_API FConcertScopedObjectTrace : public FNoncopyable
+	class FConcertScopedObjectTrace : public FNoncopyable
 	{
 		const bool bShouldTrace;
 		
@@ -23,8 +23,8 @@ namespace UE::ConcertTrace
 		
 	public:
 
-		FConcertScopedObjectTrace(uint8 InProtocol, uint32 InSequenceId, const ANSICHAR* InEventName, FSoftObjectPath InObjectPath);
-		~FConcertScopedObjectTrace();
+		CONCERTTRANSPORT_API FConcertScopedObjectTrace(uint8 InProtocol, uint32 InSequenceId, const ANSICHAR* InEventName, FSoftObjectPath InObjectPath);
+		CONCERTTRANSPORT_API ~FConcertScopedObjectTrace();
 	};
 }
-#endif
\ No newline at end of file
+#endif
diff --git a/Engine/Plugins/Developer/Concert/ConcertSync/ConcertSyncCore/Source/ConcertSyncCore/Public/ConcertDataStore.h b/Engine/Plugins/Developer/Concert/ConcertSync/ConcertSyncCore/Source/ConcertSyncCore/Public/ConcertDataStore.h
--- a/Engine/Plugins/Developer/Concert/ConcertSync/ConcertSyncCore/Source/ConcertSyncCore/Public/ConcertDataStore.h
+++ b/Engine/Plugins/Developer/Concert/ConcertSync/ConcertSyncCore/Source/ConcertSyncCore/Public/ConcertDataStore.h
@@ -4,6 +4,8 @@
 
 #include "ConcertDataStoreMessages.h"
 
+#define UE_API CONCERTSYNCCORE_API
+
 struct FConcertSessionSerializedPayload;
 
 /**
@@ -45,7 +47,7 @@ struct FConcertSessionSerializedPayload;
  * }
  * @endcode
  */
-class CONCERTSYNCCORE_API FConcertDataStore
+class FConcertDataStore
 {
 public:
 	/** Defines how the store updates an existing value. */
@@ -79,7 +81,7 @@ public:
 	 *     - EConcertDataStoreResultCode::Fetched if the key was already taken and type matched. The result holds a pointer on the fetched value.
 	 *     - EConcertDataStoreResultCode::TypeMismatch if the key was already taken but the types did not match. The result holds a null pointer.
 	 */
-	FConcertDataStoreResult FetchOrAdd(const FName& Key, const FName& TypeName, const FConcertSessionSerializedPayload& SerializedValue);
+	UE_API FConcertDataStoreResult FetchOrAdd(const FName& Key, const FName& TypeName, const FConcertSessionSerializedPayload& SerializedValue);
 
 	/**
 	 * Looks up the specified key, if it doesn't exist yet, adds new key/value pair, else, if the stored value type matches
@@ -95,7 +97,7 @@ public:
 	 *     - EConcertDataStoreResultCode::TypeMismatch if the key was already taken but the value types did not match. The result holds a null pointer.
 	 * @see[The example in the class documentation.](@ref FConcertDataStore)
 	 */
-	FConcertDataStoreResult Store(const FName& Key, const FName& TypeName, const FConcertSessionSerializedPayload& SerializedValue, const TOptional<uint32>& Version = TOptional<uint32>());
+	UE_API FConcertDataStoreResult Store(const FName& Key, const FName& TypeName, const FConcertSessionSerializedPayload& SerializedValue, const TOptional<uint32>& Version = TOptional<uint32>());
 
 	/**
 	 * Looks up the specified key, if found and types match, fetches the corresponding value. If the key is not found
@@ -109,22 +111,22 @@ public:
 	 *     - EConcertDataStoreResultCode::TypeMismatch if the key was found, but the requested type did not match the stored type. The result holds a null pointer.
 	 * @see[The example in the class documentation.](@ref FConcertDataStore)
 	 */
-	FConcertDataStoreResult Fetch(const FName& Key, const FName& TypeName) const;
+	UE_API FConcertDataStoreResult Fetch(const FName& Key, const FName& TypeName) const;
 
 	/**
 	 * Visits all the key/value currently stored.
 	 * @param Visitor The function invoked for each key/value pair.
 	 */
-	void Visit(const TFunctionRef<void(const FName&, const FConcertDataStore_StoreValue&)>& Visitor) const;
+	UE_API void Visit(const TFunctionRef<void(const FName&, const FConcertDataStore_StoreValue&)>& Visitor) const;
 
 	/** Returns the version of the specified key if found. */
-	TOptional<uint32> GetVersion(const FName& Key) const;
+	UE_API TOptional<uint32> GetVersion(const FName& Key) const;
 
 	/** Returns the number of key/value pairs currently stored. */
 	int32 GetSize() const { return KeyValueMap.Num(); }
 
 private:
-	FConcertDataStoreResult InternalAdd(const FName& Key, const FName& TypeName, const FConcertSessionSerializedPayload& Value, uint32 Version);
+	UE_API FConcertDataStoreResult InternalAdd(const FName& Key, const FName& TypeName, const FConcertSessionSerializedPayload& Value, uint32 Version);
 
 	/** Defines how the values are updated in the store. */
 	EUpdatePolicy UpdatePolicy;
@@ -132,3 +134,5 @@ private:
 	/** Maps keyName/keyValue. */
 	TMap<FName, FConcertDataStoreValueRef> KeyValueMap;
 };
+
+#undef UE_API
diff --git a/Engine/Plugins/Developer/Concert/ConcertSync/ConcertSyncCore/Source/ConcertSyncCore/Public/ConcertFileCache.h b/Engine/Plugins/Developer/Concert/ConcertSync/ConcertSyncCore/Source/ConcertSyncCore/Public/ConcertFileCache.h
--- a/Engine/Plugins/Developer/Concert/ConcertSync/ConcertSyncCore/Source/ConcertSyncCore/Public/ConcertFileCache.h
+++ b/Engine/Plugins/Developer/Concert/ConcertSync/ConcertSyncCore/Source/ConcertSyncCore/Public/ConcertFileCache.h
@@ -6,15 +6,17 @@
 #include "Misc/DateTime.h"
 #include "Templates/SharedPointer.h"
 
+#define UE_API CONCERTSYNCCORE_API
+
 /**
  * Maintains an in-memory cache of file data, attempting to keep the cache within the given upper memory threshold.
  * @note Will automatically re-cache files if they are changed on disk (via a timestamp change), and un-cache files that are deleted from disk.
  */
-class CONCERTSYNCCORE_API FConcertFileCache
+class FConcertFileCache
 {
 public:
-	FConcertFileCache(const int32 InMinimumNumberOfFilesToCache, const uint64 InMaximumNumberOfBytesToCache);
-	~FConcertFileCache();
+	UE_API FConcertFileCache(const int32 InMinimumNumberOfFilesToCache, const uint64 InMaximumNumberOfBytesToCache);
+	UE_API ~FConcertFileCache();
 
 	/**
 	 * Non-copyable.
@@ -27,42 +29,42 @@ public:
 	 * @note This function ignores the current cached file state, so will re-cache the file even if it has already been cached.
 	 * @return True if the file was cached, false otherwise.
 	 */
-	bool CacheFile(const FString& InFilename);
+	UE_API bool CacheFile(const FString& InFilename);
 
 	/**
 	 * Save the given file and add it to the cache.
 	 * @return True if the file was cached, false otherwise.
 	 */
-	bool SaveAndCacheFile(const FString& InFilename, TArray<uint8>&& InFileData);
+	UE_API bool SaveAndCacheFile(const FString& InFilename, TArray<uint8>&& InFileData);
 
 	/**
 	 * Uncache the given file.
 	 */
-	void UncacheFile(const FString& InFilename);
+	UE_API void UncacheFile(const FString& InFilename);
 
 	/**
 	 * Find or cache the given file, and get its data.
 	 * @return True if the file was found or cached, false otherwise.
 	 */
-	bool FindOrCacheFile(const FString& InFilename, TArray<uint8>& OutFileData);
+	UE_API bool FindOrCacheFile(const FString& InFilename, TArray<uint8>& OutFileData);
 
 	/**
 	 * Find the given file, and get its data.
 	 * @return True if the file was found, false otherwise.
 	 */
-	bool FindFile(const FString& InFilename, TArray<uint8>& OutFileData) const;
+	UE_API bool FindFile(const FString& InFilename, TArray<uint8>& OutFileData) const;
 
 	/**
 	 * Is the given file cached?
 	 * @return True if the file is cached, false otherwise.
 	 */
-	bool HasCachedFile(const FString& InFilename) const;
+	UE_API bool HasCachedFile(const FString& InFilename) const;
 
 private:
 	/**
 	 * Trim the cache to attempt to keep it within the requested upper memory threshold.
 	 */
-	void TrimCache();
+	UE_API void TrimCache();
 
 	/** Internal cache entry that scopes counting the cached file size into the outer file-cache */
 	class FInternalCacheEntry
@@ -116,3 +118,5 @@ private:
 	uint64 TotalCachedFileDataBytes;
 	TLruCache<FString, TSharedPtr<FInternalCacheEntry>> InternalCache;
 };
+
+#undef UE_API
diff --git a/Engine/Plugins/Developer/Concert/ConcertSync/ConcertSyncCore/Source/ConcertSyncCore/Public/ConcertLocalFileSharingService.h b/Engine/Plugins/Developer/Concert/ConcertSync/ConcertSyncCore/Source/ConcertSyncCore/Public/ConcertLocalFileSharingService.h
--- a/Engine/Plugins/Developer/Concert/ConcertSync/ConcertSyncCore/Source/ConcertSyncCore/Public/ConcertLocalFileSharingService.h
+++ b/Engine/Plugins/Developer/Concert/ConcertSync/ConcertSyncCore/Source/ConcertSyncCore/Public/ConcertLocalFileSharingService.h
@@ -4,6 +4,8 @@
 
 #include "IConcertFileSharingService.h"
 
+#define UE_API CONCERTSYNCCORE_API
+
 /**
  * Share files using a local directory. It works only for one client and one server running on the same machine. The publisher of
  * a file stores a file for sharing and gets a unique ID to consume the file back. It is the publisher responsibility to share the
@@ -12,7 +14,7 @@
  * 
  * @note This is mainly designed for the recovery system.
  */
-class CONCERTSYNCCORE_API FConcertLocalFileSharingService : public IConcertFileSharingService
+class FConcertLocalFileSharingService : public IConcertFileSharingService
 {
 public:
 	/**
@@ -21,30 +23,32 @@ public:
 	 * @note The service uses the Engine intermediate directory to share the files. The files are automatically deleted after consumption. In case of crash
 	 *       the temporary shared files are deleted on next reboot if no other client/server on the machine are not running.
 	 */
-	FConcertLocalFileSharingService(const FString& Role);
-	virtual ~FConcertLocalFileSharingService();
+	UE_API FConcertLocalFileSharingService(const FString& Role);
+	UE_API virtual ~FConcertLocalFileSharingService();
 
 	/** Copies the specified file in the service storage and return a unique ID required to consume it back. The stored file expires in 1 hour. */
-	virtual bool Publish(const FString& Pathname, FString& OutFileId) override;
+	UE_API virtual bool Publish(const FString& Pathname, FString& OutFileId) override;
 
 	/** Copies the content of the archive in the service storage and return a unique ID required to consume it back. The stored file expires in 1 hour. */
-	virtual bool Publish(FArchive& SrcAr, int64 Size, FString& OutFileId) override;
+	UE_API virtual bool Publish(FArchive& SrcAr, int64 Size, FString& OutFileId) override;
 
 	/** Consumes a file corresponding to the unique file ID, if available. The file is deleted from the service storage when the shared pointer goes out of scope. */
-	virtual TSharedPtr<FArchive> CreateReader(const FString& InFileId) override;
+	UE_API virtual TSharedPtr<FArchive> CreateReader(const FString& InFileId) override;
 
 private:
 	/** Load the list of concurrent file sharing service running across all processes. */
-	void LoadActiveServices(TArray<uint32>& OutPids);
+	UE_API void LoadActiveServices(TArray<uint32>& OutPids);
 
 	/** Saved the list of concurrent file sharing service running across all processes. */
-	void SaveActiveServices(const TArray<uint32>& InPids);
+	UE_API void SaveActiveServices(const TArray<uint32>& InPids);
 
 	/** Track active services instances across all processes and clean up any left over files that expired if no services are running anymore. */
-	void RemoveDeadProcessesAndFiles(TArray<uint32>& InOutPids);
+	UE_API void RemoveDeadProcessesAndFiles(TArray<uint32>& InOutPids);
 
 private:
 	FString SharedRootPathname;
 	FString SystemMutexName;
 	FString ActiveServicesRepositoryPathname;
 };
+
+#undef UE_API
diff --git a/Engine/Plugins/Developer/Concert/ConcertSync/ConcertSyncCore/Source/ConcertSyncCore/Public/ConcertSyncArchives.h b/Engine/Plugins/Developer/Concert/ConcertSync/ConcertSyncCore/Source/ConcertSyncCore/Public/ConcertSyncArchives.h
--- a/Engine/Plugins/Developer/Concert/ConcertSync/ConcertSyncCore/Source/ConcertSyncCore/Public/ConcertSyncArchives.h
+++ b/Engine/Plugins/Developer/Concert/ConcertSync/ConcertSyncCore/Source/ConcertSyncCore/Public/ConcertSyncArchives.h
@@ -6,6 +6,8 @@
 #include "IdentifierTable/ConcertTransportArchives.h"
 #include "UObject/SoftObjectPath.h"
 
+#define UE_API CONCERTSYNCCORE_API
+
 struct FLazyObjectPtr;
 struct FObjectPtr;
 struct FSoftObjectPtr;
@@ -68,7 +70,7 @@ CONCERTSYNCCORE_API const FSoftObjectPath& GetSkipObjectPath();
 }
 
 /** Util to handle remapping objects within the source world to be the equivalent objects in the destination world */
-class CONCERTSYNCCORE_API FConcertSyncWorldRemapper
+class FConcertSyncWorldRemapper
 {
 public:
 	FConcertSyncWorldRemapper() = default;
@@ -79,11 +81,11 @@ public:
 	{
 	}
 
-	FString RemapObjectPathName(const FString& InObjectPathName) const;
+	UE_API FString RemapObjectPathName(const FString& InObjectPathName) const;
 
-	bool ObjectBelongsToWorld(const FString& InObjectPathName) const;
+	UE_API bool ObjectBelongsToWorld(const FString& InObjectPathName) const;
 
-	bool HasMapping() const;
+	UE_API bool HasMapping() const;
 
 	FConcertSyncRemapObjectPath RemapDelegate;
 	FConcertSyncObjectPathBelongsToWorld ObjectPathBelongsToWorldDelegate;
@@ -94,27 +96,27 @@ private:
 };
 
 /** Archive for writing objects in a way that they can be sent to another instance via Concert */
-class CONCERTSYNCCORE_API FConcertSyncObjectWriter : public FConcertIdentifierWriter
+class FConcertSyncObjectWriter : public FConcertIdentifierWriter
 {
 public:
 
-	FConcertSyncObjectWriter(FConcertLocalIdentifierTable* InLocalIdentifierTable, UObject* InObj, TArray<uint8>& OutBytes, const bool InIncludeEditorOnlyData, const bool InSkipAssets, const FConcertSyncRemapObjectPath& InRemapDelegate);
-	FConcertSyncObjectWriter(FConcertLocalIdentifierTable* InLocalIdentifierTable, UObject* InObj, TArray<uint8>& OutBytes, const bool InIncludeEditorOnlyData, const bool InSkipAssets);
+	UE_API FConcertSyncObjectWriter(FConcertLocalIdentifierTable* InLocalIdentifierTable, UObject* InObj, TArray<uint8>& OutBytes, const bool InIncludeEditorOnlyData, const bool InSkipAssets, const FConcertSyncRemapObjectPath& InRemapDelegate);
+	UE_API FConcertSyncObjectWriter(FConcertLocalIdentifierTable* InLocalIdentifierTable, UObject* InObj, TArray<uint8>& OutBytes, const bool InIncludeEditorOnlyData, const bool InSkipAssets);
 
-	void SerializeObject(const UObject* InObject, const TArray<const FProperty*>* InPropertiesToWrite = nullptr, bool bAllowOuters = false);
-	void SerializeProperty(const FProperty* InProp, const UObject* InObject);
+	UE_API void SerializeObject(const UObject* InObject, const TArray<const FProperty*>* InPropertiesToWrite = nullptr, bool bAllowOuters = false);
+	UE_API void SerializeProperty(const FProperty* InProp, const UObject* InObject);
 
 	using FConcertIdentifierWriter::operator<<; // For visibility of the overloads we don't override
 
 	//~ Begin FArchive Interface
-	virtual FArchive& operator<<(UObject*& Obj) override;
-	virtual FArchive& operator<<(FLazyObjectPtr& LazyObjectPtr) override;
-	virtual FArchive& operator<<(FObjectPtr& Obj) override;
-	virtual FArchive& operator<<(FSoftObjectPtr& AssetPtr) override;
-	virtual FArchive& operator<<(FSoftObjectPath& AssetPtr) override;
-	virtual FArchive& operator<<(FWeakObjectPtr& Value) override;
-	virtual FString GetArchiveName() const override;
-	virtual bool ShouldSkipProperty(const FProperty* InProperty) const override;
+	UE_API virtual FArchive& operator<<(UObject*& Obj) override;
+	UE_API virtual FArchive& operator<<(FLazyObjectPtr& LazyObjectPtr) override;
+	UE_API virtual FArchive& operator<<(FObjectPtr& Obj) override;
+	UE_API virtual FArchive& operator<<(FSoftObjectPtr& AssetPtr) override;
+	UE_API virtual FArchive& operator<<(FSoftObjectPath& AssetPtr) override;
+	UE_API virtual FArchive& operator<<(FWeakObjectPtr& Value) override;
+	UE_API virtual FString GetArchiveName() const override;
+	UE_API virtual bool ShouldSkipProperty(const FProperty* InProperty) const override;
 	//~ End FArchive Interface
 
 	void SetSerializeNestedObjects(bool bInSerializeNestedObjects)
@@ -132,25 +134,25 @@ private:
 };
 
 /** Archive for reading objects that have been received from another instance via Concert */
-class CONCERTSYNCCORE_API FConcertSyncObjectReader : public FConcertIdentifierReader
+class FConcertSyncObjectReader : public FConcertIdentifierReader
 {
 public:
-	FConcertSyncObjectReader(const FConcertLocalIdentifierTable* InLocalIdentifierTable, FConcertSyncWorldRemapper InWorldRemapper, const FConcertSessionVersionInfo* InVersionInfo, UObject* InObj, const TArray<uint8>& InBytes, const FConcertSyncEncounteredMissingObject& InEncounteredMissingObjectDelegate);
-	FConcertSyncObjectReader(const FConcertLocalIdentifierTable* InLocalIdentifierTable, FConcertSyncWorldRemapper InWorldRemapper, const FConcertSessionVersionInfo* InVersionInfo, UObject* InObj, const TArray<uint8>& InBytes);
+	UE_API FConcertSyncObjectReader(const FConcertLocalIdentifierTable* InLocalIdentifierTable, FConcertSyncWorldRemapper InWorldRemapper, const FConcertSessionVersionInfo* InVersionInfo, UObject* InObj, const TArray<uint8>& InBytes, const FConcertSyncEncounteredMissingObject& InEncounteredMissingObjectDelegate);
+	UE_API FConcertSyncObjectReader(const FConcertLocalIdentifierTable* InLocalIdentifierTable, FConcertSyncWorldRemapper InWorldRemapper, const FConcertSessionVersionInfo* InVersionInfo, UObject* InObj, const TArray<uint8>& InBytes);
 
-	void SerializeObject(UObject* InObject);
-	void SerializeProperty(const FProperty* InProp, UObject* InObject);
+	UE_API void SerializeObject(UObject* InObject);
+	UE_API void SerializeProperty(const FProperty* InProp, UObject* InObject);
 
 	using FConcertIdentifierReader::operator<<; // For visibility of the overloads we don't override
 
 	//~ Begin FArchive Interface
-	virtual FArchive& operator<<(UObject*& Obj) override;
-	virtual FArchive& operator<<(FLazyObjectPtr& LazyObjectPtr) override;
-	virtual FArchive& operator<<(FObjectPtr& Obj) override;
-	virtual FArchive& operator<<(FSoftObjectPtr& AssetPtr) override;
-	virtual FArchive& operator<<(FSoftObjectPath& AssetPtr) override;
-	virtual FArchive& operator<<(FWeakObjectPtr& Value) override;
-	virtual FString GetArchiveName() const override;
+	UE_API virtual FArchive& operator<<(UObject*& Obj) override;
+	UE_API virtual FArchive& operator<<(FLazyObjectPtr& LazyObjectPtr) override;
+	UE_API virtual FArchive& operator<<(FObjectPtr& Obj) override;
+	UE_API virtual FArchive& operator<<(FSoftObjectPtr& AssetPtr) override;
+	UE_API virtual FArchive& operator<<(FSoftObjectPath& AssetPtr) override;
+	UE_API virtual FArchive& operator<<(FWeakObjectPtr& Value) override;
+	UE_API virtual FString GetArchiveName() const override;
 	//~ End FArchive Interface
 
 	void SetSerializeNestedObjects(bool bInSerializeNestedObjects)
@@ -168,24 +170,24 @@ private:
 };
 
 /** Archive for rewriting identifiers (currently names) so that they belong to a different identifier table */
-class CONCERTSYNCCORE_API FConcertSyncObjectRewriter : public FConcertIdentifierRewriter
+class FConcertSyncObjectRewriter : public FConcertIdentifierRewriter
 {
 public:
-	FConcertSyncObjectRewriter(const FConcertLocalIdentifierTable* InLocalIdentifierTable, FConcertLocalIdentifierTable* InRewriteIdentifierTable, const FConcertSessionVersionInfo* InVersionInfo, TArray<uint8>& InBytes);
+	UE_API FConcertSyncObjectRewriter(const FConcertLocalIdentifierTable* InLocalIdentifierTable, FConcertLocalIdentifierTable* InRewriteIdentifierTable, const FConcertSessionVersionInfo* InVersionInfo, TArray<uint8>& InBytes);
 
-	void RewriteObject(const UClass* InClass);
-	void RewriteProperty(const FProperty* InProp);
+	UE_API void RewriteObject(const UClass* InClass);
+	UE_API void RewriteProperty(const FProperty* InProp);
 
 	using FConcertIdentifierRewriter::operator<<; // For visibility of the overloads we don't override
 
 	//~ Begin FArchive Interface
-	virtual FArchive& operator<<(UObject*& Obj) override;
-	virtual FArchive& operator<<(FLazyObjectPtr& LazyObjectPtr) override;
-	virtual FArchive& operator<<(FObjectPtr& Obj) override;
-	virtual FArchive& operator<<(FSoftObjectPtr& AssetPtr) override;
-	virtual FArchive& operator<<(FSoftObjectPath& AssetPtr) override;
-	virtual FArchive& operator<<(FWeakObjectPtr& Value) override;
-	virtual FString GetArchiveName() const override;
+	UE_API virtual FArchive& operator<<(UObject*& Obj) override;
+	UE_API virtual FArchive& operator<<(FLazyObjectPtr& LazyObjectPtr) override;
+	UE_API virtual FArchive& operator<<(FObjectPtr& Obj) override;
+	UE_API virtual FArchive& operator<<(FSoftObjectPtr& AssetPtr) override;
+	UE_API virtual FArchive& operator<<(FSoftObjectPath& AssetPtr) override;
+	UE_API virtual FArchive& operator<<(FWeakObjectPtr& Value) override;
+	UE_API virtual FString GetArchiveName() const override;
 	//~ End FArchive Interface
 
 protected:
@@ -195,3 +197,5 @@ protected:
 	/** Called with the result of OnObjectSerialized to get the object path to set on the in-memory object value (the output of operator<<) */
 	virtual FSoftObjectPath GetOutputObjectPath(const FSoftObjectPath& Obj) { return Obj; }
 };
+
+#undef UE_API
diff --git a/Engine/Plugins/Developer/Concert/ConcertSync/ConcertSyncCore/Source/ConcertSyncCore/Public/ConcertSyncSessionDatabase.h b/Engine/Plugins/Developer/Concert/ConcertSync/ConcertSyncCore/Source/ConcertSyncCore/Public/ConcertSyncSessionDatabase.h
--- a/Engine/Plugins/Developer/Concert/ConcertSync/ConcertSyncCore/Source/ConcertSyncCore/Public/ConcertSyncSessionDatabase.h
+++ b/Engine/Plugins/Developer/Concert/ConcertSync/ConcertSyncCore/Source/ConcertSyncCore/Public/ConcertSyncSessionDatabase.h
@@ -6,6 +6,8 @@
 #include "Misc/EBreakBehavior.h"
 #include "Templates/PimplPtr.h"
 
+#define UE_API CONCERTSYNCCORE_API
+
 class FStructOnScope;
 enum class EConcertSyncActivityEventType : uint8;
 enum class EConcertSyncActivityFlags : uint8;
@@ -46,23 +48,23 @@ DECLARE_MULTICAST_DELEGATE_OneParam(FOnPackageSaved, const FName& PackageName);
  * Database of activities that have happened in a Concert Sync Session.
  * Stores the activity index and their associated data.
  */
-class CONCERTSYNCCORE_API FConcertSyncSessionDatabase
+class FConcertSyncSessionDatabase
 {
 public:
 	
-	FConcertSyncSessionDatabase();
-	~FConcertSyncSessionDatabase();
+	UE_API FConcertSyncSessionDatabase();
+	UE_API ~FConcertSyncSessionDatabase();
 
 	FConcertSyncSessionDatabase(const FConcertSyncSessionDatabase&) = delete;
 	FConcertSyncSessionDatabase& operator=(const FConcertSyncSessionDatabase&) = delete;
 
-	FConcertSyncSessionDatabase(FConcertSyncSessionDatabase&&);
-	FConcertSyncSessionDatabase& operator=(FConcertSyncSessionDatabase&&);
+	UE_API FConcertSyncSessionDatabase(FConcertSyncSessionDatabase&&);
+	UE_API FConcertSyncSessionDatabase& operator=(FConcertSyncSessionDatabase&&);
 
 	/**
 	 * Is this a valid database? (ie, has been successfully opened).
 	 */
-	bool IsValid() const;
+	UE_API bool IsValid() const;
 
 	/**
 	 * Open (or create) a database file.
@@ -71,7 +73,7 @@ public:
 	 *
 	 * @return True if the database was opened, false otherwise.
 	 */
-	bool Open(const FString& InSessionPath);
+	UE_API bool Open(const FString& InSessionPath);
 
 	/**
 	 * Open (or create) a database file.
@@ -81,25 +83,25 @@ public:
 	 *
 	 * @return True if the database was opened, false otherwise.
 	 */
-	bool Open(const FString& InSessionPath, const ESQLiteDatabaseOpenMode InOpenMode);
+	UE_API bool Open(const FString& InSessionPath, const ESQLiteDatabaseOpenMode InOpenMode);
 
 	/**
 	 * Close an open database file.
 	 *
 	 * @param InDeleteDatabase			True if the session database and its associated data should be deleted after closing the database.
 	 */
-	bool Close(const bool InDeleteDatabase = false);
+	UE_API bool Close(const bool InDeleteDatabase = false);
 
 	/**
 	 * Get the filename of the currently open database, or an empty string.
 	 * @note The returned filename will be an absolute pathname.
 	 */
-	FString GetFilename() const;
+	UE_API FString GetFilename() const;
 
 	/**
 	 * Get the last error reported by this database.
 	 */
-	FString GetLastError() const;
+	UE_API FString GetLastError() const;
 
 	/**
 	 * Add a new connection activity to this database, assigning it both an activity and connection event ID.
@@ -112,7 +114,7 @@ public:
 	 *
 	 * @return True if the connection activity was added, false otherwise.
 	 */
-	bool AddConnectionActivity(const FConcertSyncConnectionActivity& InConnectionActivity, int64& OutActivityId, int64& OutConnectionEventId);
+	UE_API bool AddConnectionActivity(const FConcertSyncConnectionActivity& InConnectionActivity, int64& OutActivityId, int64& OutConnectionEventId);
 
 	/**
 	 * Add a new lock activity to this database, assigning it both an activity and lock event ID.
@@ -125,7 +127,7 @@ public:
 	 *
 	 * @return True if the lock activity was added, false otherwise.
 	 */
-	bool AddLockActivity(const FConcertSyncLockActivity& InLockActivity, int64& OutActivityId, int64& OutLockEventId);
+	UE_API bool AddLockActivity(const FConcertSyncLockActivity& InLockActivity, int64& OutActivityId, int64& OutLockEventId);
 
 	/**
 	 * Add a new transaction activity to this database, assigning it both an activity and transaction event ID.
@@ -138,7 +140,7 @@ public:
 	 *
 	 * @return True if the transaction activity was added, false otherwise.
 	 */
-	bool AddTransactionActivity(const FConcertSyncTransactionActivity& InTransactionActivity, int64& OutActivityId, int64& OutTransactionEventId);
+	UE_API bool AddTransactionActivity(const FConcertSyncTransactionActivity& InTransactionActivity, int64& OutActivityId, int64& OutTransactionEventId);
 
 	/**
 	 * Add a new package activity to this database, assigning it both an activity and package event ID.
@@ -153,7 +155,7 @@ public:
 	 *
 	 * @return True if the package activity was added, false otherwise.
 	 */
-	bool AddPackageActivity(const FConcertSyncActivity& InPackageActivity, const FConcertPackageInfo& InPackageInfo, FConcertPackageDataStream& InPackageDataStream, int64& OutActivityId, int64& OutPackageEventId);
+	UE_API bool AddPackageActivity(const FConcertSyncActivity& InPackageActivity, const FConcertPackageInfo& InPackageInfo, FConcertPackageDataStream& InPackageDataStream, int64& OutActivityId, int64& OutPackageEventId);
 
 	/**
 	 * Add a new replication activity to this database, assigning it both an activity and lock event ID.
@@ -166,12 +168,12 @@ public:
 	 *
 	 * @return True if the lock activity was added, false otherwise.
 	 */
-	bool AddReplicationActivity(const FConcertSyncReplicationActivity& InReplicationActivity, int64& OutActivityId, int64& OutReplicationEventId);
+	UE_API bool AddReplicationActivity(const FConcertSyncReplicationActivity& InReplicationActivity, int64& OutActivityId, int64& OutReplicationEventId);
 	
 	/**
 	 * Iterates the given range, calls UpdateCallback on each element, and commits the update.
 	 */
-	bool SetActivities(const TSet<int64>& ActivityIds, TFunctionRef<void(FConcertSyncActivity&)> UpdateCallback);
+	UE_API bool SetActivities(const TSet<int64>& ActivityIds, TFunctionRef<void(FConcertSyncActivity&)> UpdateCallback);
 	
 	/**
 	 * Set a connection activity in this database, creating or replacing it.
@@ -182,7 +184,7 @@ public:
 	 *
 	 * @return True if the connection activity was set, false otherwise.
 	 */
-	bool SetConnectionActivity(const FConcertSyncConnectionActivity& InConnectionActivity);
+	UE_API bool SetConnectionActivity(const FConcertSyncConnectionActivity& InConnectionActivity);
 
 	/**
 	 * Set a lock activity in this database, creating or replacing it.
@@ -193,7 +195,7 @@ public:
 	 *
 	 * @return True if the lock activity was set, false otherwise.
 	 */
-	bool SetLockActivity(const FConcertSyncLockActivity& InLockActivity);
+	UE_API bool SetLockActivity(const FConcertSyncLockActivity& InLockActivity);
 
 	/**
 	 * Set a transaction activity in this database, creating or replacing it.
@@ -205,7 +207,7 @@ public:
 	 *
 	 * @return True if the transaction activity was set, false otherwise.
 	 */
-	bool SetTransactionActivity(const FConcertSyncTransactionActivity& InTransactionActivity, const bool bMetaDataOnly = false);
+	UE_API bool SetTransactionActivity(const FConcertSyncTransactionActivity& InTransactionActivity, const bool bMetaDataOnly = false);
 
 	/**
 	 * Set a package activity in this database, creating or replacing it.
@@ -218,7 +220,7 @@ public:
 	 *
 	 * @return True if the package activity was set, false otherwise.
 	 */
-	bool SetPackageActivity(const FConcertSyncActivity& InPackageActivity, FConcertSyncPackageEventData& InPackageActivityEvent, const bool bMetaDataOnly = false);
+	UE_API bool SetPackageActivity(const FConcertSyncActivity& InPackageActivity, FConcertSyncPackageEventData& InPackageActivityEvent, const bool bMetaDataOnly = false);
 
 	/**
 	 * Set a replication activity in this database, creating or replacing it.
@@ -229,7 +231,7 @@ public:
 	 *
 	 * @return True if the transaction activity was set, false otherwise.
 	 */
-	bool SetReplicationActivity(const FConcertSyncReplicationActivity& InReplicationActivity);
+	UE_API bool SetReplicationActivity(const FConcertSyncReplicationActivity& InReplicationActivity);
 	
 	/**
 	 * Get the generic part of an activity from this database.
@@ -239,7 +241,7 @@ public:
 	 *
 	 * @return True if the activity was found, false otherwise.
 	 */
-	bool GetActivity(const int64 InActivityId, FConcertSyncActivity& OutActivity) const;
+	UE_API bool GetActivity(const int64 InActivityId, FConcertSyncActivity& OutActivity) const;
 
 	/**
 	 * Get a connection activity from this database.
@@ -249,7 +251,7 @@ public:
 	 *
 	 * @return True if the connection activity was found, false otherwise.
 	 */
-	bool GetConnectionActivity(const int64 InActivityId, FConcertSyncConnectionActivity& OutConnectionActivity) const;
+	UE_API bool GetConnectionActivity(const int64 InActivityId, FConcertSyncConnectionActivity& OutConnectionActivity) const;
 
 	/**
 	 * Get a lock activity from this database.
@@ -259,7 +261,7 @@ public:
 	 *
 	 * @return True if the lock activity was found, false otherwise.
 	 */
-	bool GetLockActivity(const int64 InActivityId, FConcertSyncLockActivity& OutLockActivity) const;
+	UE_API bool GetLockActivity(const int64 InActivityId, FConcertSyncLockActivity& OutLockActivity) const;
 
 	/**
 	 * Get a transaction activity from this database.
@@ -269,7 +271,7 @@ public:
 	 *
 	 * @return True if the transaction activity was found, false otherwise.
 	 */
-	bool GetTransactionActivity(const int64 InActivityId, FConcertSyncTransactionActivity& OutTransactionActivity) const;
+	UE_API bool GetTransactionActivity(const int64 InActivityId, FConcertSyncTransactionActivity& OutTransactionActivity) const;
 
 	/**
 	 * Get a package activity from this database.
@@ -280,7 +282,7 @@ public:
 	 *
 	 * @return True if the package activity was found, false otherwise.
 	 */
-	bool GetPackageActivity(const int64 InActivityId, FConsumePackageActivityFunc PackageActivityFn) const;
+	UE_API bool GetPackageActivity(const int64 InActivityId, FConsumePackageActivityFunc PackageActivityFn) const;
 
 	/**
 	 * Get a replication activity from this database.
@@ -290,7 +292,7 @@ public:
 	 *
 	 * @return True if the replication activity was found, false otherwise.
 	 */
-	bool GetReplicationActivity(const int64 InActivityId, FConcertSyncReplicationActivity& OutReplicationActivity) const;
+	UE_API bool GetReplicationActivity(const int64 InActivityId, FConcertSyncReplicationActivity& OutReplicationActivity) const;
 
 	/**
 	 * Get a replication activity from this database by its event ID
@@ -300,7 +302,7 @@ public:
 	 *
 	 * @return True if the replication activity was found, false otherwise.
 	 */
-	bool GetReplicationActivityForEvent(const int64 InEventId, FConcertSyncReplicationActivity& OutReplicationActivity) const;
+	UE_API bool GetReplicationActivityForEvent(const int64 InEventId, FConcertSyncReplicationActivity& OutReplicationActivity) const;
 	
 	/**
 	 * Get the type of an activity in this database.
@@ -310,7 +312,7 @@ public:
 	 *
 	 * @return True if the activity was found, false otherwise.
 	 */
-	bool GetActivityEventType(const int64 InActivityId, EConcertSyncActivityEventType& OutEventType) const;
+	UE_API bool GetActivityEventType(const int64 InActivityId, EConcertSyncActivityEventType& OutEventType) const;
 	
 	/**
 	 * Get the generic part of an activity for an event in this database.
@@ -321,7 +323,7 @@ public:
 	 *
 	 * @return True if the activity was found, false otherwise.
 	 */
-	bool GetActivityForEvent(const int64 InEventId, const EConcertSyncActivityEventType InEventType, FConcertSyncActivity& OutActivity) const;
+	UE_API bool GetActivityForEvent(const int64 InEventId, const EConcertSyncActivityEventType InEventType, FConcertSyncActivity& OutActivity) const;
 
 	/**
 	 * Get a connection activity for an event in this database.
@@ -331,7 +333,7 @@ public:
 	 *
 	 * @return True if the connection activity was found, false otherwise.
 	 */
-	bool GetConnectionActivityForEvent(const int64 InConnectionEventId, FConcertSyncConnectionActivity& OutConnectionActivity) const;
+	UE_API bool GetConnectionActivityForEvent(const int64 InConnectionEventId, FConcertSyncConnectionActivity& OutConnectionActivity) const;
 
 	/**
 	 * Get a lock activity for an event in this database.
@@ -341,7 +343,7 @@ public:
 	 *
 	 * @return True if the lock activity was found, false otherwise.
 	 */
-	bool GetLockActivityForEvent(const int64 InLockEventId, FConcertSyncLockActivity& OutLockActivity) const;
+	UE_API bool GetLockActivityForEvent(const int64 InLockEventId, FConcertSyncLockActivity& OutLockActivity) const;
 
 	/**
 	 * Get a transaction activity for an event in this database.
@@ -351,7 +353,7 @@ public:
 	 *
 	 * @return True if the transaction activity was found, false otherwise.
 	 */
-	bool GetTransactionActivityForEvent(const int64 InTransactionEventId, FConcertSyncTransactionActivity& OutTransactionActivity) const;
+	UE_API bool GetTransactionActivityForEvent(const int64 InTransactionEventId, FConcertSyncTransactionActivity& OutTransactionActivity) const;
 
 	/**
 	 * Get a package activity for an event in this database.
@@ -362,7 +364,7 @@ public:
 	 *
 	 * @return True if the package activity was found, false otherwise.
 	 */
-	bool GetPackageActivityForEvent(const int64 InPackageEventId, FIteratePackageActivityFunc PackageActivityFn) const;
+	UE_API bool GetPackageActivityForEvent(const int64 InPackageEventId, FIteratePackageActivityFunc PackageActivityFn) const;
 
 	/**
 	 * Enumerate the generic part of the activities in this database.
@@ -371,7 +373,7 @@ public:
 	 *
 	 * @return True if the activities were enumerated without error, false otherwise.
 	 */
-	bool EnumerateActivities(FIterateActivityFunc InCallback) const;
+	UE_API bool EnumerateActivities(FIterateActivityFunc InCallback) const;
 
 	/**
 	 * Enumerate all the connection activities in this database.
@@ -380,7 +382,7 @@ public:
 	 *
 	 * @return True if the connection activities were enumerated without error, false otherwise.
 	 */
-	bool EnumerateConnectionActivities(TFunctionRef<bool(FConcertSyncConnectionActivity&&)> InCallback) const;
+	UE_API bool EnumerateConnectionActivities(TFunctionRef<bool(FConcertSyncConnectionActivity&&)> InCallback) const;
 
 	/**
 	 * Enumerate all the lock activities in this database.
@@ -389,7 +391,7 @@ public:
 	 *
 	 * @return True if the lock activities were enumerated without error, false otherwise.
 	 */
-	bool EnumerateLockActivities(TFunctionRef<bool(FConcertSyncLockActivity&&)> InCallback) const;
+	UE_API bool EnumerateLockActivities(TFunctionRef<bool(FConcertSyncLockActivity&&)> InCallback) const;
 
 	/**
 	 * Enumerate all the transaction activities in this database.
@@ -398,7 +400,7 @@ public:
 	 *
 	 * @return True if the transaction activities were enumerated without error, false otherwise.
 	 */
-	bool EnumerateTransactionActivities(TFunctionRef<bool(FConcertSyncTransactionActivity&&)> InCallback) const;
+	UE_API bool EnumerateTransactionActivities(TFunctionRef<bool(FConcertSyncTransactionActivity&&)> InCallback) const;
 
 	/**
 	 * Enumerate all the package activities in this database.
@@ -407,7 +409,7 @@ public:
 	 *
 	 * @return True if the package activities were enumerated without error, false otherwise.
 	 */
-	bool EnumeratePackageActivities(FIteratePackageActivityFunc InCallback) const;
+	UE_API bool EnumeratePackageActivities(FIteratePackageActivityFunc InCallback) const;
 
 	/**
 	 * Enumerate all the replication activities in this database.
@@ -416,7 +418,7 @@ public:
 	 *
 	 * @return True if the replication activities were enumerated without error, false otherwise.
 	 */
-	bool EnumerateReplicationActivities(TFunctionRef<bool(FConcertSyncReplicationActivity&&)> InCallback) const;
+	UE_API bool EnumerateReplicationActivities(TFunctionRef<bool(FConcertSyncReplicationActivity&&)> InCallback) const;
 
 	/**
 	 * Enumerate all the activities in this database of the given type.
@@ -426,7 +428,7 @@ public:
 	 *
 	 * @return True if the activities were enumerated without error, false otherwise.
 	 */
-	bool EnumerateActivitiesForEventType(const EConcertSyncActivityEventType InEventType, FIterateActivityFunc InCallback) const;
+	UE_API bool EnumerateActivitiesForEventType(const EConcertSyncActivityEventType InEventType, FIterateActivityFunc InCallback) const;
 
 	/**
 	 * Enumerate all the activities in this database in the given range.
@@ -437,7 +439,7 @@ public:
 	 *
 	 * @return True if the activities were enumerated without error, false otherwise.
 	 */
-	bool EnumerateActivitiesInRange(const int64 InFirstActivityId, const int64 InMaxNumActivities, FIterateActivityFunc InCallback) const;
+	UE_API bool EnumerateActivitiesInRange(const int64 InFirstActivityId, const int64 InMaxNumActivities, FIterateActivityFunc InCallback) const;
 
 	/**
 	 * Enumerate the IDs and event types of all the activities in this database.
@@ -446,7 +448,7 @@ public:
 	 *
 	 * @return True if the activities were enumerated without error, false otherwise.
 	 */
-	bool EnumerateActivityIdsAndEventTypes(TFunctionRef<bool(int64, EConcertSyncActivityEventType)> InCallback) const;
+	UE_API bool EnumerateActivityIdsAndEventTypes(TFunctionRef<bool(int64, EConcertSyncActivityEventType)> InCallback) const;
 
 	/**
 	 * Enumerate the IDs and event types of the activities in this database in the given range.
@@ -457,7 +459,7 @@ public:
 	 *
 	 * @return True if the activities were enumerated without error, false otherwise.
 	 */
-	bool EnumerateActivityIdsAndEventTypesInRange(const int64 InFirstActivityId, const int64 InMaxNumActivities, TFunctionRef<bool(int64, EConcertSyncActivityEventType)> InCallback) const;
+	UE_API bool EnumerateActivityIdsAndEventTypesInRange(const int64 InFirstActivityId, const int64 InMaxNumActivities, TFunctionRef<bool(int64, EConcertSyncActivityEventType)> InCallback) const;
 
 	/**
 	 * Enumerate the IDs, event types and flags of the activities in this database in the given range.
@@ -468,7 +470,7 @@ public:
 	 *
 	 * @return True if the activities were enumerated without error, false otherwise.
 	 */
-	bool EnumerateActivityIdsWithEventTypesAndFlagsInRange(const int64 InFirstActivityId, const int64 InMaxNumActivities, TFunctionRef<bool(int64, EConcertSyncActivityEventType, EConcertSyncActivityFlags)> InCallback) const;
+	UE_API bool EnumerateActivityIdsWithEventTypesAndFlagsInRange(const int64 InFirstActivityId, const int64 InMaxNumActivities, TFunctionRef<bool(int64, EConcertSyncActivityEventType, EConcertSyncActivityFlags)> InCallback) const;
 	
 	/**
 	 * Get the maximum ID of the activities in this database.
@@ -477,7 +479,7 @@ public:
 	 *
 	 * @return True if the activity ID was resolved, false otherwise.
 	 */
-	bool GetActivityMaxId(int64& OutActivityId) const;
+	UE_API bool GetActivityMaxId(int64& OutActivityId) const;
 
 	/**
 	 * Set an endpoint in this database, creating or replacing it.
@@ -487,7 +489,7 @@ public:
 	 *
 	 * @return True if the endpoint was set, false otherwise.
 	 */
-	bool SetEndpoint(const FGuid& InEndpointId, const FConcertSyncEndpointData& InEndpointData);
+	UE_API bool SetEndpoint(const FGuid& InEndpointId, const FConcertSyncEndpointData& InEndpointData);
 
 	/**
 	 * Get an endpoint from this database.
@@ -497,7 +499,7 @@ public:
 	 *
 	 * @return True if the endpoint was found, false otherwise.
 	 */
-	bool GetEndpoint(const FGuid& InEndpointId, FConcertSyncEndpointData& OutEndpointData) const;
+	UE_API bool GetEndpoint(const FGuid& InEndpointId, FConcertSyncEndpointData& OutEndpointData) const;
 
 	/**
 	 * Enumerate all the endpoints in this database.
@@ -506,7 +508,7 @@ public:
 	 *
 	 * @return True if the endpoints were enumerated without error, false otherwise.
 	 */
-	bool EnumerateEndpoints(TFunctionRef<bool(FConcertSyncEndpointIdAndData&&)> InCallback) const;
+	UE_API bool EnumerateEndpoints(TFunctionRef<bool(FConcertSyncEndpointIdAndData&&)> InCallback) const;
 
 	/**
 	 * Enumerate all the endpoint IDs in this database.
@@ -515,7 +517,7 @@ public:
 	 *
 	 * @return True if the endpoint IDs were enumerated without error, false otherwise.
 	 */
-	bool EnumerateEndpointIds(TFunctionRef<bool(FGuid)> InCallback) const;
+	UE_API bool EnumerateEndpointIds(TFunctionRef<bool(FGuid)> InCallback) const;
 
 	/**
 	 * Get a connection event from this database.
@@ -525,7 +527,7 @@ public:
 	 *
 	 * @return True if the connection event was found, false otherwise.
 	 */
-	bool GetConnectionEvent(const int64 InConnectionEventId, FConcertSyncConnectionEvent& OutConnectionEvent) const;
+	UE_API bool GetConnectionEvent(const int64 InConnectionEventId, FConcertSyncConnectionEvent& OutConnectionEvent) const;
 
 	/**
 	 * Get a lock event from this database.
@@ -535,7 +537,7 @@ public:
 	 *
 	 * @return True if the lock event was found, false otherwise.
 	 */
-	bool GetLockEvent(const int64 InLockEventId, FConcertSyncLockEvent& OutLockEvent) const;
+	UE_API bool GetLockEvent(const int64 InLockEventId, FConcertSyncLockEvent& OutLockEvent) const;
 
 	/**
 	 * Get a transaction event from this database.
@@ -546,7 +548,7 @@ public:
 	 *
 	 * @return True if the transaction event was found, false otherwise.
 	 */
-	bool GetTransactionEvent(const int64 InTransactionEventId, FConcertSyncTransactionEvent& OutTransactionEvent, const bool InMetaDataOnly = false) const;
+	UE_API bool GetTransactionEvent(const int64 InTransactionEventId, FConcertSyncTransactionEvent& OutTransactionEvent, const bool InMetaDataOnly = false) const;
 
 	/**
 	 * Get a replication event from this database.
@@ -556,7 +558,7 @@ public:
 	 *
 	 * @return True if the replication event was found, false otherwise.
 	 */
-	bool GetReplicationEvent(const int64 InReplicationEventId, FConcertSyncReplicationEvent& OutReplicationEvent) const;
+	UE_API bool GetReplicationEvent(const int64 InReplicationEventId, FConcertSyncReplicationEvent& OutReplicationEvent) const;
 
 	/**
 	 * Get the maximum ID of the transaction events in this database.
@@ -565,7 +567,7 @@ public:
 	 *
 	 * @return True if the transaction event ID was resolved, false otherwise.
 	 */
-	bool GetTransactionMaxEventId(int64& OutTransactionEventId) const;
+	UE_API bool GetTransactionMaxEventId(int64& OutTransactionEventId) const;
 
 	/**
 	 * Check whether the given transaction event ID is currently for a live transaction event.
@@ -575,7 +577,7 @@ public:
 	 *
 	 * @return True if the transaction event ID was queried successfully, false otherwise.
 	 */
-	bool IsLiveTransactionEvent(const int64 InTransactionEventId, bool& OutIsLive) const;
+	UE_API bool IsLiveTransactionEvent(const int64 InTransactionEventId, bool& OutIsLive) const;
 
 	/**
 	 * Get the IDs of any live transaction events.
@@ -584,7 +586,7 @@ public:
 	 *
 	 * @return True if the live transaction event IDs were resolved, false otherwise.
 	 */
-	bool GetLiveTransactionEventIds(TArray<int64>& OutTransactionEventIds) const;
+	UE_API bool GetLiveTransactionEventIds(TArray<int64>& OutTransactionEventIds) const;
 
 	/**
 	 * Get the IDs of any live transaction events for the given package name.
@@ -594,7 +596,7 @@ public:
 	 *
 	 * @return True if the live transaction event IDs were resolved, false otherwise.
 	 */
-	bool GetLiveTransactionEventIdsForPackage(const FName InPackageName, TArray<int64>& OutTransactionEventIds) const;
+	UE_API bool GetLiveTransactionEventIdsForPackage(const FName InPackageName, TArray<int64>& OutTransactionEventIds) const;
 
 	/**
 	 * Get if a Package has any live transactions
@@ -604,7 +606,7 @@ public:
 	 *
 	 * @return True if query was resolved correctly, false otherwise.
 	 */
-	bool PackageHasLiveTransactions(const FName InPackageName, bool& OutHasLiveTransaction) const;
+	UE_API bool PackageHasLiveTransactions(const FName InPackageName, bool& OutHasLiveTransaction) const;
 
 	/**
 	 * Enumerate the IDs of any live transaction events for the given package name.
@@ -614,7 +616,7 @@ public:
 	 *
 	 * @return True if the live transaction event IDs were enumerated without error, false otherwise.
 	 */
-	bool EnumerateLiveTransactionEventIdsForPackage(const FName InPackageName, TFunctionRef<bool(int64)> InCallback) const;
+	UE_API bool EnumerateLiveTransactionEventIdsForPackage(const FName InPackageName, TFunctionRef<bool(int64)> InCallback) const;
 
 	/**
 	 * Enumerate the names of of any packages that have live transaction events.
@@ -623,7 +625,7 @@ public:
 	 *
 	 * @return True if the package names were enumerated without error, false otherwise.
 	 */
-	bool EnumeratePackageNamesWithLiveTransactions(TFunctionRef<bool(FName)> InCallback) const;
+	UE_API bool EnumeratePackageNamesWithLiveTransactions(TFunctionRef<bool(FName)> InCallback) const;
 
 	/**
 	 * Add a new dummy package event to this database, assigning it a package event ID.
@@ -634,7 +636,7 @@ public:
 	 *
 	 * @return True if the package event was added, false otherwise.
 	 */
-	bool AddDummyPackageEvent(const FName InPackageName, int64& OutPackageEventId);
+	UE_API bool AddDummyPackageEvent(const FName InPackageName, int64& OutPackageEventId);
 
 	/**
 	 * Get a package event meta data (omitting the package data itself) from this database.
@@ -645,7 +647,7 @@ public:
 	 *
 	 * @return True if the package event was found, false otherwise.
 	 */
-	bool GetPackageEventMetaData(const int64 InPackageEventId, int64& OuptPackageRevision, FConcertPackageInfo& OutPackageInfo) const;
+	UE_API bool GetPackageEventMetaData(const int64 InPackageEventId, int64& OuptPackageRevision, FConcertPackageInfo& OutPackageInfo) const;
 
 	/**
 	 * Get Package Event Id from Package Name.  The event id will be from the head revision of the package.
@@ -655,7 +657,7 @@ public:
 	 *
 	 * @return True if the package event was found, false otherwise
 	 */
-	bool GetPackageHeadEventFromName(const FName InPackageName, int64& OutPackageEventId);
+	UE_API bool GetPackageHeadEventFromName(const FName InPackageName, int64& OutPackageEventId);
 
 	/**
 	 * Get a package event from this database.
@@ -665,7 +667,7 @@ public:
 	 *
 	 * @return True if the package event was found, false otherwise.
 	 */
-	bool GetPackageEvent(const int64 InPackageEventId, const TFunctionRef<void(FConcertSyncPackageEventData&)>& PackageEventFn) const;
+	UE_API bool GetPackageEvent(const int64 InPackageEventId, const TFunctionRef<void(FConcertSyncPackageEventData&)>& PackageEventFn) const;
 
 	/**
 	 * Enumerate package names for packages with a head revision (at least one package event)
@@ -675,7 +677,7 @@ public:
 	 *
 	 * @return True if the package data was enumerated without error, false otherwise.
 	 */
-	bool EnumeratePackageNamesWithHeadRevision(TFunctionRef<bool(FName)> InCallback, bool IgnorePersisted) const;
+	UE_API bool EnumeratePackageNamesWithHeadRevision(TFunctionRef<bool(FName)> InCallback, bool IgnorePersisted) const;
 
 	/**
 	 * Enumerate the head revision package data for all packages in this database.
@@ -684,7 +686,7 @@ public:
 	 *
 	 * @return True if the package data was enumerated without error, false otherwise.
 	 */
-	bool EnumerateHeadRevisionPackageData(TFunctionRef<bool(const FConcertPackageInfo&, FConcertPackageDataStream&)> InCallback) const;
+	UE_API bool EnumerateHeadRevisionPackageData(TFunctionRef<bool(const FConcertPackageInfo&, FConcertPackageDataStream&)> InCallback) const;
 
 	/**
 	 * Get the data from this database for the given package name for the given revision.
@@ -695,7 +697,7 @@ public:
 	 *
 	 * @return True if package data could be found for the given revision, false otherwise.
 	 */
-	bool GetPackageInfoForRevision(const FName InPackageName, FConcertPackageInfo& OutPackageInfo, const int64* InPackageRevision = nullptr) const;
+	UE_API bool GetPackageInfoForRevision(const FName InPackageName, FConcertPackageInfo& OutPackageInfo, const int64* InPackageRevision = nullptr) const;
 	
 	/**
 	 * Get the data from this database for the given package name for the given revision.
@@ -706,7 +708,7 @@ public:
 	 *
 	 * @return True if package data could be found for the given revision, false otherwise.
 	 */
-	bool GetPackageDataForRevision(const FName InPackageName, const TFunctionRef<void(const FConcertPackageInfo&, FConcertPackageDataStream&)>& InCallback, const int64* InPackageRevision = nullptr) const;
+	UE_API bool GetPackageDataForRevision(const FName InPackageName, const TFunctionRef<void(const FConcertPackageInfo&, FConcertPackageDataStream&)>& InCallback, const int64* InPackageRevision = nullptr) const;
 
 	/**
 	 * Gets the package size in bytes of a package
@@ -716,7 +718,7 @@ public:
 	 *
 	 * @return True if package data could be found for the given revision, false otherwise.
 	 */
-	TOptional<int64> GetPackageSizeForRevision(const FName InPackageName, const int64* InPackageRevision = nullptr) const;
+	UE_API TOptional<int64> GetPackageSizeForRevision(const FName InPackageName, const int64* InPackageRevision = nullptr) const;
 	
 	/**
 	 * Get the head revision in this database for the given package name.
@@ -726,7 +728,7 @@ public:
 	 *
 	 * @return True if the head revision was resolved, false otherwise.
 	 */
-	bool GetPackageHeadRevision(const FName InPackageName, int64& OutRevision) const;
+	UE_API bool GetPackageHeadRevision(const FName InPackageName, int64& OutRevision) const;
 
 	/**
 	 * Check whether the given package event ID is currently for the head revision of the package in the event.
@@ -736,7 +738,7 @@ public:
 	 *
 	 * @return True if the package event ID was queried successfully, false otherwise.
 	 */
-	bool IsHeadRevisionPackageEvent(const int64 InPackageEventId, bool& OutIsHeadRevision) const;
+	UE_API bool IsHeadRevisionPackageEvent(const int64 InPackageEventId, bool& OutIsHeadRevision) const;
 
 	/**
 	 * Get the maximum ID of the replication events in this database that was triggered by EndpointId and had ActivityType.
@@ -747,7 +749,7 @@ public:
 	 *
 	 * @return True if the replication event ID was resolved, false otherwise.
 	 */
-	bool GetReplicationMaxEventIdByClientAndType(const FGuid& EndpointId, EConcertSyncReplicationActivityType ActivityType, int64& OutReplicationEventId) const;
+	UE_API bool GetReplicationMaxEventIdByClientAndType(const FGuid& EndpointId, EConcertSyncReplicationActivityType ActivityType, int64& OutReplicationEventId) const;
 
 	/**
 	 * Add a package event ID for the head revision to the persist events in this database, if not already existing.
@@ -757,7 +759,7 @@ public:
 	 *
 	 * @return True if the persist event was added, false otherwise.
 	 */
-	bool AddPersistEventForHeadRevision(FName InPackageName, int64& OutPersistEventId);
+	UE_API bool AddPersistEventForHeadRevision(FName InPackageName, int64& OutPersistEventId);
 
 	/**
 	 * Update the specified transaction event.
@@ -766,7 +768,7 @@ public:
 	 * @param InTransactionEvent		The event to store.
 	 * @return True if the transaction was updated.
 	 */
-	bool UpdateTransactionEvent(const int64 InTransactionEventId, const FConcertSyncTransactionEvent& InTransactionEvent);
+	UE_API bool UpdateTransactionEvent(const int64 InTransactionEventId, const FConcertSyncTransactionEvent& InTransactionEvent);
 
 	/**
 	 * Update the specified package event.
@@ -775,22 +777,22 @@ public:
 	 * @param PackageEvent				The event to store.
 	 * @return True if the package was updated.
 	 */
-	bool UpdatePackageEvent(const int64 InPackageEventId, FConcertSyncPackageEventData& PackageEvent);
+	UE_API bool UpdatePackageEvent(const int64 InPackageEventId, FConcertSyncPackageEventData& PackageEvent);
 
 	/**
 	 * Check Asynchronous Tasks Status
 	 */
-	void UpdateAsynchronousTasks();
+	UE_API void UpdateAsynchronousTasks();
 
 	/**
 	 * Flush any ongoing asynchronous tasks.
 	 */
-	void FlushAsynchronousTasks();
+	UE_API void FlushAsynchronousTasks();
 
 	/**
 	 * Has all asynchronous package tasks completed.
 	 */
-	bool HasWritePackageTasksCompleted() const;
+	UE_API bool HasWritePackageTasksCompleted() const;
 
 	FOnActivityProduced& OnActivityProduced() { return ActivityProducedEvent; }
 	
@@ -798,10 +800,10 @@ private:
 
 	using FProcessPackageRequest = TFunctionRef<bool(const FConcertPackageInfo& PackageInfo, const FString& DataFilename)>;
 	/** Helper function which obtains package information and passes it to HandleFunc. */
-	bool HandleRequestPackageRequest(const FName InPackageName, const int64* InPackageRevision, FProcessPackageRequest HandleFunc) const;
+	UE_API bool HandleRequestPackageRequest(const FName InPackageName, const int64* InPackageRevision, FProcessPackageRequest HandleFunc) const;
 
 	/** Helper functions that for getting a package revision for an optional package revision argument */
-	TOptional<int64> GetSpecifiedOrHeadPackageRevision(FName InPackageName, const int64* InPackageRevision) const;
+	UE_API TOptional<int64> GetSpecifiedOrHeadPackageRevision(FName InPackageName, const int64* InPackageRevision) const;
 
 	/**
 	 * Schedule an asynchronous write for the given Package Stream.  The stream must be in-memory. File sharing
@@ -811,7 +813,7 @@ private:
 	 * @param InDstPackageBlobPathName  Full path of the destination package.
 	 * @param InPackageDataStream       The package data stream.
 	 **/
-	void ScheduleAsyncWrite(const FName& InPackageName, const FString& InDstPackageBlobPathname, FConcertPackageDataStream& InPackageDataStream);
+	UE_API void ScheduleAsyncWrite(const FName& InPackageName, const FString& InDstPackageBlobPathname, FConcertPackageDataStream& InPackageDataStream);
 
 	/**
 	 * Set the active ignored state for the given activity.
@@ -821,7 +823,7 @@ private:
 	 *
 	 * @return True if the ignored state was set, false otherwise.
 	 */
-	bool SetActivityIgnoredState(const int64 InActivityId, const bool InIsIgnored);
+	UE_API bool SetActivityIgnoredState(const int64 InActivityId, const bool InIsIgnored);
 
 	/**
 	 * Add a new connection event to this database, assigning it a connection event ID.
@@ -831,7 +833,7 @@ private:
 	 *
 	 * @return True if the connection event was added, false otherwise.
 	 */
-	bool AddConnectionEvent(const FConcertSyncConnectionEvent& InConnectionEvent, int64& OutConnectionEventId);
+	UE_API bool AddConnectionEvent(const FConcertSyncConnectionEvent& InConnectionEvent, int64& OutConnectionEventId);
 
 	/**
 	 * Set a connection event in this database, creating or replacing it.
@@ -841,7 +843,7 @@ private:
 	 *
 	 * @return True if the connection event was set, false otherwise.
 	 */
-	bool SetConnectionEvent(const int64 InConnectionEventId, const FConcertSyncConnectionEvent& InConnectionEvent);
+	UE_API bool SetConnectionEvent(const int64 InConnectionEventId, const FConcertSyncConnectionEvent& InConnectionEvent);
 
 	/**
 	 * Add a new lock event to this database, assigning it a lock event ID.
@@ -851,7 +853,7 @@ private:
 	 *
 	 * @return True if the lock event was added, false otherwise.
 	 */
-	bool AddLockEvent(const FConcertSyncLockEvent& InLockEvent, int64& OutLockEventId);
+	UE_API bool AddLockEvent(const FConcertSyncLockEvent& InLockEvent, int64& OutLockEventId);
 
 	/**
 	 * Set a lock event in this database, creating or replacing it.
@@ -861,7 +863,7 @@ private:
 	 *
 	 * @return True if the lock event was set, false otherwise.
 	 */
-	bool SetLockEvent(const int64 InLockEventId, const FConcertSyncLockEvent& InLockEvent);
+	UE_API bool SetLockEvent(const int64 InLockEventId, const FConcertSyncLockEvent& InLockEvent);
 
 	/**
 	 * Add a new transaction event to this database, assigning it a transaction event ID.
@@ -871,7 +873,7 @@ private:
 	 *
 	 * @return True if the transaction event was added, false otherwise.
 	 */
-	bool AddTransactionEvent(const FConcertSyncTransactionEvent& InTransactionEvent, int64& OutTransactionEventId);
+	UE_API bool AddTransactionEvent(const FConcertSyncTransactionEvent& InTransactionEvent, int64& OutTransactionEventId);
 
 	/**
 	 * Set a transaction event in this database, creating or replacing it.
@@ -882,7 +884,7 @@ private:
 	 *
 	 * @return True if the transaction event was set, false otherwise.
 	 */
-	bool SetTransactionEvent(const int64 InTransactionEventId, const FConcertSyncTransactionEvent& InTransactionEvent, const bool bMetaDataOnly = false);
+	UE_API bool SetTransactionEvent(const int64 InTransactionEventId, const FConcertSyncTransactionEvent& InTransactionEvent, const bool bMetaDataOnly = false);
 
 	/**
 	 * Enumerate the IDs of any live transaction events for the given package name ID.
@@ -892,7 +894,7 @@ private:
 	 *
 	 * @return True if the live transaction event IDs were enumerated without error, false otherwise.
 	 */
-	bool EnumerateLiveTransactionEventIdsForPackage(const int64 InPackageNameId, TFunctionRef<bool(int64)> InCallback) const;
+	UE_API bool EnumerateLiveTransactionEventIdsForPackage(const int64 InPackageNameId, TFunctionRef<bool(int64)> InCallback) const;
 
 	/**
 	 * Add a new package event to this database, assigning it a package event ID.
@@ -903,7 +905,7 @@ private:
 	 *
 	 * @return True if the package event was added, false otherwise.
 	 */
-	bool AddPackageEvent(const FConcertPackageInfo& InPackageInfo, FConcertPackageDataStream& InPackageDataStream, int64& OutPackageEventId);
+	UE_API bool AddPackageEvent(const FConcertPackageInfo& InPackageInfo, FConcertPackageDataStream& InPackageDataStream, int64& OutPackageEventId);
 
 	/**
 	 * Set a package event in this database, creating or replacing it.
@@ -915,7 +917,7 @@ private:
 	 *
 	 * @return True if the package event was set, false otherwise.
 	 */
-	bool SetPackageEvent(const int64 InPackageEventId, int64 PackageRevision, const FConcertPackageInfo& PackageInfo, FConcertPackageDataStream* InPackageDataStream);
+	UE_API bool SetPackageEvent(const int64 InPackageEventId, int64 PackageRevision, const FConcertPackageInfo& PackageInfo, FConcertPackageDataStream* InPackageDataStream);
 
 	/**
 	 * Add a new replication event to this database, assigning it a replication event ID.
@@ -925,7 +927,7 @@ private:
 	 *
 	 * @return True if the replication event was added, false otherwise.
 	 */
-	bool AddReplicationEvent(const FConcertSyncReplicationEvent& InReplicationEvent, int64& OutReplicationEventId);
+	UE_API bool AddReplicationEvent(const FConcertSyncReplicationEvent& InReplicationEvent, int64& OutReplicationEventId);
 
 	/**
 	 * Set a replication event in this database, creating or replacing it.
@@ -935,7 +937,7 @@ private:
 	 *
 	 * @return True if the replication event was set, false otherwise.
 	 */
-	bool SetReplicationEvent(const int64 InReplicationEventId, const FConcertSyncReplicationEvent& InReplicationEvent);
+	UE_API bool SetReplicationEvent(const int64 InReplicationEventId, const FConcertSyncReplicationEvent& InReplicationEvent);
 	
 	/**
 	 * Get the maximum ID of the package events in this database.
@@ -944,7 +946,7 @@ private:
 	 *
 	 * @return True if the package event ID was resolved, false otherwise.
 	 */
-	bool GetPackageMaxEventId(int64& OutPackageEventId) const;
+	UE_API bool GetPackageMaxEventId(int64& OutPackageEventId) const;
 
 	/**
 	 * Get the object path name from this database that matches the given object name ID.
@@ -954,7 +956,7 @@ private:
 	 *
 	 * @return True if the object path name was resolved, false otherwise.
 	 */
-	bool GetObjectPathName(const int64 InObjectNameId, FName& OutObjectPathName) const;
+	UE_API bool GetObjectPathName(const int64 InObjectNameId, FName& OutObjectPathName) const;
 
 	/**
 	 * Get the object name ID from this database that matches the given object path name.
@@ -964,7 +966,7 @@ private:
 	 *
 	 * @return True if the object name ID name was resolved, false otherwise.
 	 */
-	bool GetObjectNameId(const FName InObjectPathName, int64& OutObjectNameId) const;
+	UE_API bool GetObjectNameId(const FName InObjectPathName, int64& OutObjectNameId) const;
 
 	/**
 	 * Ensure that an object name ID in this database matches the given object path name, creating it if required.
@@ -974,7 +976,7 @@ private:
 	 *
 	 * @return True if the object name ID name was resolved, false otherwise.
 	 */
-	bool EnsureObjectNameId(const FName InObjectPathName, int64& OutObjectNameId);
+	UE_API bool EnsureObjectNameId(const FName InObjectPathName, int64& OutObjectNameId);
 
 	/**
 	 * Get the package name from this database that matches the given package name ID.
@@ -984,7 +986,7 @@ private:
 	 *
 	 * @return True if the package name was resolved, false otherwise.
 	 */
-	bool GetPackageName(const int64 InPackageNameId, FName& OutPackageName) const;
+	UE_API bool GetPackageName(const int64 InPackageNameId, FName& OutPackageName) const;
 
 	/**
 	 * Get the package name ID from this database that matches the given package name.
@@ -994,7 +996,7 @@ private:
 	 *
 	 * @return True if the package name ID name was resolved, false otherwise.
 	 */
-	bool GetPackageNameId(const FName InPackageName, int64& OutPackageNameId) const;
+	UE_API bool GetPackageNameId(const FName InPackageName, int64& OutPackageNameId) const;
 
 	/**
 	 * Ensure that a package name ID in this database matches the given package name, creating it if required.
@@ -1004,7 +1006,7 @@ private:
 	 *
 	 * @return True if the package name ID name was resolved, false otherwise.
 	 */
-	bool EnsurePackageNameId(const FName InPackageName, int64& OutPackageNameId);
+	UE_API bool EnsurePackageNameId(const FName InPackageName, int64& OutPackageNameId);
 
 	/**
 	 * Map the resource names so that they have an association with the given lock event ID in this database.
@@ -1014,7 +1016,7 @@ private:
 	 *
 	 * @return True if the resource name references were mapped, false otherwise.
 	 */
-	bool MapResourceNamesForLock(const int64 InLockEventId, const TArray<FName>& InResourceNames);
+	UE_API bool MapResourceNamesForLock(const int64 InLockEventId, const TArray<FName>& InResourceNames);
 
 	/**
 	 * Map the package names in the given transaction so that they have an association with the given transaction event ID in this database.
@@ -1024,7 +1026,7 @@ private:
 	 *
 	 * @return True if the package name references were mapped, false otherwise.
 	 */
-	bool MapPackageNamesForTransaction(const int64 InTransactionEventId, const FConcertTransactionFinalizedEvent& InTransactionEvent);
+	UE_API bool MapPackageNamesForTransaction(const int64 InTransactionEventId, const FConcertTransactionFinalizedEvent& InTransactionEvent);
 
 	/**
 	 * Map the object path names in the given transaction so that they have an association with the given transaction event ID in this database.
@@ -1034,7 +1036,7 @@ private:
 	 *
 	 * @return True if the object path name references were mapped, false otherwise.
 	 */
-	bool MapObjectNamesForTransaction(const int64 InTransactionEventId, const FConcertTransactionFinalizedEvent& InTransactionEvent);
+	UE_API bool MapObjectNamesForTransaction(const int64 InTransactionEventId, const FConcertTransactionFinalizedEvent& InTransactionEvent);
 
 	/**
 	 * Save and cache the given transaction data with the given filename.
@@ -1044,7 +1046,7 @@ private:
 	 *
 	 * @return True if the transaction data was saved, false otherwise.
 	 */
-	bool SaveTransaction(const FString& InTransactionFilename, const FStructOnScope& InTransaction) const;
+	UE_API bool SaveTransaction(const FString& InTransactionFilename, const FStructOnScope& InTransaction) const;
 
 	/**
 	 * Load and cache the transaction data for the given filename.
@@ -1054,7 +1056,7 @@ private:
 	 *
 	 * @return True if the transaction data was loaded, false otherwise.
 	 */
-	bool LoadTransaction(const FString& InTransactionFilename, FStructOnScope& OutTransaction) const;
+	UE_API bool LoadTransaction(const FString& InTransactionFilename, FStructOnScope& OutTransaction) const;
 
 	/**
 	 * Save and cache the given package data with the given filename.
@@ -1065,7 +1067,7 @@ private:
 	 *
 	 * @return True if the package data was saved, false otherwise.
 	 */
-	bool SavePackage(const FString& InDstPackageBlobPathname, const FConcertPackageInfo& InPackageInfo, FConcertPackageDataStream& InPackageDataStream);
+	UE_API bool SavePackage(const FString& InDstPackageBlobPathname, const FConcertPackageInfo& InPackageInfo, FConcertPackageDataStream& InPackageDataStream);
 
 	/**
 	 * Load the package data for the given filename.
@@ -1075,7 +1077,7 @@ private:
 	 * @note The archive passed as output is only valid during the callback, its position is set at the beginning of the data (not necessarily zero) and may contain data passed the supplied size.
 	 * @return True if the package data was loaded, false otherwise.
 	 */
-	bool LoadPackage(const FString& InPackageBlobFilename, const TFunctionRef<void(FConcertPackageDataStream&)>& PackageDataStreamFn) const;
+	UE_API bool LoadPackage(const FString& InPackageBlobFilename, const TFunctionRef<void(FConcertPackageDataStream&)>& PackageDataStreamFn) const;
 	
 	/** Called when an activity is produced */
 	FOnActivityProduced ActivityProducedEvent;
@@ -1130,3 +1132,5 @@ namespace ConcertSyncSessionDatabaseFilterUtil
 	CONCERTSYNCCORE_API bool PackageEventPassesFilter(const int64 InPackageEventId, const FConcertSessionFilter& InSessionFilter, const FConcertSyncSessionDatabase& InDatabase);
 
 } // namespace ConcertSyncSessionDatabaseFilterUtil
+
+#undef UE_API
diff --git a/Engine/Plugins/Developer/Concert/ConcertSync/ConcertSyncCore/Source/ConcertSyncCore/Public/ConcertSyncSessionTypes.h b/Engine/Plugins/Developer/Concert/ConcertSync/ConcertSyncCore/Source/ConcertSyncCore/Public/ConcertSyncSessionTypes.h
--- a/Engine/Plugins/Developer/Concert/ConcertSync/ConcertSyncCore/Source/ConcertSyncCore/Public/ConcertSyncSessionTypes.h
+++ b/Engine/Plugins/Developer/Concert/ConcertSync/ConcertSyncCore/Source/ConcertSyncCore/Public/ConcertSyncSessionTypes.h
@@ -7,6 +7,8 @@
 #include "UObject/StructOnScope.h"
 #include "ConcertSyncSessionTypes.generated.h"
 
+#define UE_API CONCERTSYNCCORE_API
+
 /** Types of connection events */
 UENUM()
 enum class EConcertSyncConnectionEventType : uint8
@@ -281,7 +283,7 @@ struct FConcertSyncPackageActivity : public FConcertSyncActivity
 
 /** Base summary for an activity entry in a Concert Sync Session */
 USTRUCT()
-struct CONCERTSYNCCORE_API FConcertSyncActivitySummary
+struct FConcertSyncActivitySummary
 {
 	GENERATED_BODY()
 
@@ -294,16 +296,16 @@ public:
 	 * @param InUserDisplayName		The display name of the user the activity belongs to, or an empty text to skip the user information.
 	 * @param InUseRichText			True if the output should be suitable for consumption as rich-text, or false for plain-text.
 	 */
-	FText ToDisplayText(const FText InUserDisplayName, const bool InUseRichText = false) const;
+	UE_API FText ToDisplayText(const FText InUserDisplayName, const bool InUseRichText = false) const;
 
 protected:
-	virtual FText CreateDisplayText(const bool InUseRichText) const;
-	virtual FText CreateDisplayTextForUser(const FText InUserDisplayName, const bool InUseRichText) const;
+	UE_API virtual FText CreateDisplayText(const bool InUseRichText) const;
+	UE_API virtual FText CreateDisplayTextForUser(const FText InUserDisplayName, const bool InUseRichText) const;
 };
 
 /** Summary for a connection activity entry in a Concert Sync Session */
 USTRUCT()
-struct CONCERTSYNCCORE_API FConcertSyncConnectionActivitySummary : public FConcertSyncActivitySummary
+struct FConcertSyncConnectionActivitySummary : public FConcertSyncActivitySummary
 {
 	GENERATED_BODY()
 
@@ -312,17 +314,17 @@ struct CONCERTSYNCCORE_API FConcertSyncConnectionActivitySummary : public FConce
 	EConcertSyncConnectionEventType ConnectionEventType = EConcertSyncConnectionEventType::Connected;
 
 	/** Create this summary from an connection event */
-	static FConcertSyncConnectionActivitySummary CreateSummaryForEvent(const FConcertSyncConnectionEvent& InEvent);
+	static UE_API FConcertSyncConnectionActivitySummary CreateSummaryForEvent(const FConcertSyncConnectionEvent& InEvent);
 
 protected:
 	//~ FConcertSyncActivitySummary interface
-	virtual FText CreateDisplayText(const bool InUseRichText) const override;
-	virtual FText CreateDisplayTextForUser(const FText InUserDisplayName, const bool InUseRichText) const override;
+	UE_API virtual FText CreateDisplayText(const bool InUseRichText) const override;
+	UE_API virtual FText CreateDisplayTextForUser(const FText InUserDisplayName, const bool InUseRichText) const override;
 };
 
 /** Summary for a lock activity entry in a Concert Sync Session */
 USTRUCT()
-struct CONCERTSYNCCORE_API FConcertSyncLockActivitySummary : public FConcertSyncActivitySummary
+struct FConcertSyncLockActivitySummary : public FConcertSyncActivitySummary
 {
 	GENERATED_BODY()
 
@@ -343,17 +345,17 @@ struct CONCERTSYNCCORE_API FConcertSyncLockActivitySummary : public FConcertSync
 	int32 NumResources = 0;
 
 	/** Create this summary from a lock event */
-	static FConcertSyncLockActivitySummary CreateSummaryForEvent(const FConcertSyncLockEvent& InEvent);
+	static UE_API FConcertSyncLockActivitySummary CreateSummaryForEvent(const FConcertSyncLockEvent& InEvent);
 
 protected:
 	//~ FConcertSyncActivitySummary interface
-	virtual FText CreateDisplayText(const bool InUseRichText) const override;
-	virtual FText CreateDisplayTextForUser(const FText InUserDisplayName, const bool InUseRichText) const override;
+	UE_API virtual FText CreateDisplayText(const bool InUseRichText) const override;
+	UE_API virtual FText CreateDisplayTextForUser(const FText InUserDisplayName, const bool InUseRichText) const override;
 };
 
 /** Summary for a transaction activity entry in a Concert Sync Session */
 USTRUCT()
-struct CONCERTSYNCCORE_API FConcertSyncTransactionActivitySummary : public FConcertSyncActivitySummary
+struct FConcertSyncTransactionActivitySummary : public FConcertSyncActivitySummary
 {
 	GENERATED_BODY()
 
@@ -382,17 +384,17 @@ struct CONCERTSYNCCORE_API FConcertSyncTransactionActivitySummary : public FConc
 	int32 NumActions = 0;
 
 	/** Create this summary from a transaction event */
-	static FConcertSyncTransactionActivitySummary CreateSummaryForEvent(const FConcertSyncTransactionEvent& InEvent);
+	static UE_API FConcertSyncTransactionActivitySummary CreateSummaryForEvent(const FConcertSyncTransactionEvent& InEvent);
 
 protected:
 	//~ FConcertSyncActivitySummary interface
-	virtual FText CreateDisplayText(const bool InUseRichText) const override;
-	virtual FText CreateDisplayTextForUser(const FText InUserDisplayName, const bool InUseRichText) const override;
+	UE_API virtual FText CreateDisplayText(const bool InUseRichText) const override;
+	UE_API virtual FText CreateDisplayTextForUser(const FText InUserDisplayName, const bool InUseRichText) const override;
 };
 
 /** Summary for a package activity entry in a Concert Sync Session */
 USTRUCT()
-struct CONCERTSYNCCORE_API FConcertSyncPackageActivitySummary : public FConcertSyncActivitySummary
+struct FConcertSyncPackageActivitySummary : public FConcertSyncActivitySummary
 {
 	GENERATED_BODY()
 
@@ -417,12 +419,12 @@ struct CONCERTSYNCCORE_API FConcertSyncPackageActivitySummary : public FConcertS
 	bool bPreSave = false;
 
 	/** Create this summary from a package event */
-	static FConcertSyncPackageActivitySummary CreateSummaryForEvent(const FConcertPackageInfo& InEvent);
+	static UE_API FConcertSyncPackageActivitySummary CreateSummaryForEvent(const FConcertPackageInfo& InEvent);
 
 protected:
 	//~ FConcertSyncActivitySummary interface
-	virtual FText CreateDisplayText(const bool InUseRichText) const override;
-	virtual FText CreateDisplayTextForUser(const FText InUserDisplayName, const bool InUseRichText) const override;
+	UE_API virtual FText CreateDisplayText(const bool InUseRichText) const override;
+	UE_API virtual FText CreateDisplayTextForUser(const FText InUserDisplayName, const bool InUseRichText) const override;
 };
 
 struct FConcertSessionActivity
@@ -458,3 +460,5 @@ struct FConcertSessionActivity
 	 */
 	TUniquePtr<FConcertSessionSerializedPayload> EventPayload;
 };
+
+#undef UE_API
diff --git a/Engine/Plugins/Developer/Concert/ConcertSync/ConcertSyncCore/Source/ConcertSyncCore/Public/ConcertSyncSettings.h b/Engine/Plugins/Developer/Concert/ConcertSync/ConcertSyncCore/Source/ConcertSyncCore/Public/ConcertSyncSettings.h
--- a/Engine/Plugins/Developer/Concert/ConcertSync/ConcertSyncCore/Source/ConcertSyncCore/Public/ConcertSyncSettings.h
+++ b/Engine/Plugins/Developer/Concert/ConcertSync/ConcertSyncCore/Source/ConcertSyncCore/Public/ConcertSyncSettings.h
@@ -8,7 +8,7 @@
 #include "ConcertSyncSettings.generated.h"
 
 USTRUCT()
-struct CONCERTSYNCCORE_API FTransactionClassFilter
+struct FTransactionClassFilter
 {
 	GENERATED_BODY()
 
@@ -26,7 +26,7 @@ struct CONCERTSYNCCORE_API FTransactionClassFilter
 };
 
 USTRUCT()
-struct CONCERTSYNCCORE_API FPackageClassFilter
+struct FPackageClassFilter
 {
 	GENERATED_BODY()
 
@@ -56,8 +56,8 @@ enum class EShowPresenceInGame
 	VisibleInGameAlways
 };
 
-UCLASS(config=ConcertSyncCore)
-class CONCERTSYNCCORE_API UConcertSyncConfig : public UObject
+UCLASS(MinimalAPI, config=ConcertSyncCore)
+class UConcertSyncConfig : public UObject
 {
 	GENERATED_BODY()
 
diff --git a/Engine/Plugins/Developer/Concert/ConcertSync/ConcertSyncCore/Source/ConcertSyncCore/Public/HistoryEdition/ActivityDependencyEdge.h b/Engine/Plugins/Developer/Concert/ConcertSync/ConcertSyncCore/Source/ConcertSyncCore/Public/HistoryEdition/ActivityDependencyEdge.h
--- a/Engine/Plugins/Developer/Concert/ConcertSync/ConcertSyncCore/Source/ConcertSyncCore/Public/HistoryEdition/ActivityDependencyEdge.h
+++ b/Engine/Plugins/Developer/Concert/ConcertSync/ConcertSyncCore/Source/ConcertSyncCore/Public/HistoryEdition/ActivityDependencyEdge.h
@@ -95,7 +95,7 @@ namespace UE::ConcertSyncCore
 
 
 	/** Describes a dependency from one activity to an earlier activity. */
-	class CONCERTSYNCCORE_API FActivityDependencyEdge
+	class FActivityDependencyEdge
 	{
 	public:
 		
diff --git a/Engine/Plugins/Developer/Concert/ConcertSync/ConcertSyncCore/Source/ConcertSyncCore/Public/HistoryEdition/ActivityDependencyGraph.h b/Engine/Plugins/Developer/Concert/ConcertSync/ConcertSyncCore/Source/ConcertSyncCore/Public/HistoryEdition/ActivityDependencyGraph.h
--- a/Engine/Plugins/Developer/Concert/ConcertSync/ConcertSyncCore/Source/ConcertSyncCore/Public/HistoryEdition/ActivityDependencyGraph.h
+++ b/Engine/Plugins/Developer/Concert/ConcertSync/ConcertSyncCore/Source/ConcertSyncCore/Public/HistoryEdition/ActivityDependencyGraph.h
@@ -4,6 +4,8 @@
 
 #include "ActivityNode.h"
 
+#define UE_API CONCERTSYNCCORE_API
+
 template <typename FuncType> class TFunctionRef;
 template <typename OptionalType> struct TOptional;
 
@@ -57,28 +59,30 @@ namespace UE::ConcertSyncCore
 	 *       I
 	 *       5
 	 */
-	class CONCERTSYNCCORE_API FActivityDependencyGraph
+	class FActivityDependencyGraph
 	{
 	public:
 		
 		/** Adds a node. Invalidates the results of GetNodeById. */
-		FActivityNodeID AddActivity(int64 ActivityIndex, EActivityNodeFlags NodeFlags = EActivityNodeFlags::None);
+		UE_API FActivityNodeID AddActivity(int64 ActivityIndex, EActivityNodeFlags NodeFlags = EActivityNodeFlags::None);
 		
 		/** Adds a dependency from an existing node to another existing node */
-		bool AddDependency(FActivityNodeID From, FActivityDependencyEdge To);
+		UE_API bool AddDependency(FActivityNodeID From, FActivityDependencyEdge To);
 		
 		/** Tries to find a corresponding node for an activity ID. */
-		TOptional<FActivityNodeID> FindNodeByActivity(int64 ActivityID) const;
+		UE_API TOptional<FActivityNodeID> FindNodeByActivity(int64 ActivityID) const;
 
 		/** Gets the node identified by the given ID, if it is valid. Only valid as long as no new node is added.*/
-		const FActivityNode& GetNodeById(FActivityNodeID ID) const;
+		UE_API const FActivityNode& GetNodeById(FActivityNodeID ID) const;
 
-		void ForEachNode(TFunctionRef<void(const FActivityNode&)> ConsumerFunc) const;
+		UE_API void ForEachNode(TFunctionRef<void(const FActivityNode&)> ConsumerFunc) const;
 		
 	private:
 
 		TArray<FActivityNode> Nodes;
 		
-		FActivityNode& GetNodeByIdInternal(FActivityNodeID ID);
+		UE_API FActivityNode& GetNodeByIdInternal(FActivityNodeID ID);
 	};
 }
+
+#undef UE_API
diff --git a/Engine/Plugins/Developer/Concert/ConcertSync/ConcertSyncCore/Source/ConcertSyncCore/Public/HistoryEdition/ActivityGraphIDs.h b/Engine/Plugins/Developer/Concert/ConcertSync/ConcertSyncCore/Source/ConcertSyncCore/Public/HistoryEdition/ActivityGraphIDs.h
--- a/Engine/Plugins/Developer/Concert/ConcertSync/ConcertSyncCore/Source/ConcertSyncCore/Public/HistoryEdition/ActivityGraphIDs.h
+++ b/Engine/Plugins/Developer/Concert/ConcertSync/ConcertSyncCore/Source/ConcertSyncCore/Public/HistoryEdition/ActivityGraphIDs.h
@@ -9,7 +9,7 @@
 namespace UE::ConcertSyncCore
 {
 	/** Enforces type-safe usage of IDs by avoiding implicit conversions. */
-	struct CONCERTSYNCCORE_API FActivityNodeID
+	struct FActivityNodeID
 	{
 		size_t ID;
 
diff --git a/Engine/Plugins/Developer/Concert/ConcertSync/ConcertSyncCore/Source/ConcertSyncCore/Public/HistoryEdition/ActivityNode.h b/Engine/Plugins/Developer/Concert/ConcertSync/ConcertSyncCore/Source/ConcertSyncCore/Public/HistoryEdition/ActivityNode.h
--- a/Engine/Plugins/Developer/Concert/ConcertSync/ConcertSyncCore/Source/ConcertSyncCore/Public/HistoryEdition/ActivityNode.h
+++ b/Engine/Plugins/Developer/Concert/ConcertSync/ConcertSyncCore/Source/ConcertSyncCore/Public/HistoryEdition/ActivityNode.h
@@ -7,6 +7,8 @@
 #include "HistoryEdition/ActivityGraphIDs.h"
 #include "Misc/Optional.h" // IWYU pragma: keep
 
+#define UE_API CONCERTSYNCCORE_API
+
 struct FScriptContainerElement;
 
 namespace UE::ConcertSyncCore
@@ -23,7 +25,7 @@ namespace UE::ConcertSyncCore
 	ENUM_CLASS_FLAGS(EActivityNodeFlags)
 
 	/** A node corresponds to an activity and can depend on preceding activity*/
-	class CONCERTSYNCCORE_API FActivityNode
+	class FActivityNode
 	{
 		friend FActivityDependencyGraph;
 	public:
@@ -35,12 +37,12 @@ namespace UE::ConcertSyncCore
 		{}
 
 		bool HasAnyDependency() const { return Dependencies.Num() > 0; }
-		bool DependsOnNode(FActivityNodeID NodeId, TOptional<EActivityDependencyReason> WithReason = {}, TOptional<EDependencyStrength> WithStrength = {}) const; 
-		bool DependsOnActivity(int64 InActivityId, const FActivityDependencyGraph& Graph, TOptional<EActivityDependencyReason> WithReason = {}, TOptional<EDependencyStrength> WithStrength = {}) const;
+		UE_API bool DependsOnNode(FActivityNodeID NodeId, TOptional<EActivityDependencyReason> WithReason = {}, TOptional<EDependencyStrength> WithStrength = {}) const; 
+		UE_API bool DependsOnActivity(int64 InActivityId, const FActivityDependencyGraph& Graph, TOptional<EActivityDependencyReason> WithReason = {}, TOptional<EDependencyStrength> WithStrength = {}) const;
 		
 		bool AffectsAnyActivity() const { return GetAffectedChildren().Num() > 0; }
 		bool AffectsNode(FActivityNodeID NodeID) const { return AffectedChildren.Contains(NodeID); }
-		bool AffectsActivity(int64 InActivityId, const FActivityDependencyGraph& Graph) const;
+		UE_API bool AffectsActivity(int64 InActivityId, const FActivityDependencyGraph& Graph) const;
 		
 		int64 GetActivityId() const { return ActivityId; }
 		FActivityNodeID GetNodeIndex() const { return NodeIndex; }
@@ -70,3 +72,5 @@ namespace UE::ConcertSyncCore
 		TArray<FActivityNodeID> AffectedChildren;
 	};
 }
+
+#undef UE_API
diff --git a/Engine/Plugins/Developer/Concert/ConcertSync/ConcertSyncCore/Source/ConcertSyncCore/Public/HistoryEdition/HistoryAnalysis.h b/Engine/Plugins/Developer/Concert/ConcertSync/ConcertSyncCore/Source/ConcertSyncCore/Public/HistoryEdition/HistoryAnalysis.h
--- a/Engine/Plugins/Developer/Concert/ConcertSync/ConcertSyncCore/Source/ConcertSyncCore/Public/HistoryEdition/HistoryAnalysis.h
+++ b/Engine/Plugins/Developer/Concert/ConcertSync/ConcertSyncCore/Source/ConcertSyncCore/Public/HistoryEdition/HistoryAnalysis.h
@@ -11,7 +11,7 @@ class FConcertSyncSessionDatabase;
 namespace UE::ConcertSyncCore
 {
 	/** Describes the activities that must be considered when deleting an activity */
-	struct CONCERTSYNCCORE_API FHistoryAnalysisResult
+	struct FHistoryAnalysisResult
 	{
 		/** The activities that must be removed. */
 		TSet<FActivityID> HardDependencies;
diff --git a/Engine/Plugins/Developer/Concert/ConcertSync/ConcertSyncCore/Source/ConcertSyncCore/Public/Misc/ObjectPathHierarchy.h b/Engine/Plugins/Developer/Concert/ConcertSync/ConcertSyncCore/Source/ConcertSyncCore/Public/Misc/ObjectPathHierarchy.h
--- a/Engine/Plugins/Developer/Concert/ConcertSync/ConcertSyncCore/Source/ConcertSyncCore/Public/Misc/ObjectPathHierarchy.h
+++ b/Engine/Plugins/Developer/Concert/ConcertSync/ConcertSyncCore/Source/ConcertSyncCore/Public/Misc/ObjectPathHierarchy.h
@@ -1,4 +1,4 @@
-﻿// Copyright Epic Games, Inc. All Rights Reserved.
+// Copyright Epic Games, Inc. All Rights Reserved.
 
 #pragma once
 
@@ -10,6 +10,8 @@
 #include "Templates/UniquePtr.h"
 #include "Templates/UnrealTemplate.h"
 
+#define UE_API CONCERTSYNCCORE_API
+
 namespace UE::ConcertSyncCore
 {
 	enum class ETreeTraversalBehavior : uint8
@@ -49,15 +51,15 @@ namespace UE::ConcertSyncCore
 	};
 	
 	/** Keeps track of the outer tree hierarchy of FSoftObjectPath. */
-	class CONCERTSYNCCORE_API FObjectPathHierarchy
+	class FObjectPathHierarchy
 		// This class cannot be shallowly copied because the nodes' referenced dynamic allocations - we don't expect instances to be copied, so we won't bother implementing a deep copy.
 		: public FNoncopyable
 	{
 	public:
 
 		FObjectPathHierarchy() = default;
-		FObjectPathHierarchy(FObjectPathHierarchy&& Other);
-		FObjectPathHierarchy& operator=(FObjectPathHierarchy&& Other);
+		UE_API FObjectPathHierarchy(FObjectPathHierarchy&& Other);
+		UE_API FObjectPathHierarchy& operator=(FObjectPathHierarchy&& Other);
 
 		/**
 		 * Traverses the hierarchy in pre-order (root first, then its children), starting at an optional parent object.
@@ -81,7 +83,7 @@ namespace UE::ConcertSyncCore
 		 * @param Callback The callback to invoke 
 		 * @param Start The object at which to start walking down the hierarchy
 		 */
-		void TraverseTopToBottom(TFunctionRef<ETreeTraversalBehavior(const FChildRelation& Relation)> Callback, const FSoftObjectPath& Start = {}) const;
+		UE_API void TraverseTopToBottom(TFunctionRef<ETreeTraversalBehavior(const FChildRelation& Relation)> Callback, const FSoftObjectPath& Start = {}) const;
 
 		/**
 		 * Traverses the hierarchy in post-order (children first, then the root), starting at an optional parent object.
@@ -105,14 +107,14 @@ namespace UE::ConcertSyncCore
 		 * @param Callback The callback to invoke 
 		 * @param Start The object at which to start walking up the hierarchy
 		 */
-		void TraverseBottomToTop(TFunctionRef<EBreakBehavior(const FChildRelation& Relation)> Callback, const FSoftObjectPath& Start = {}) const;
+		UE_API void TraverseBottomToTop(TFunctionRef<EBreakBehavior(const FChildRelation& Relation)> Callback, const FSoftObjectPath& Start = {}) const;
 
 		/** @return A valid TOptional if the object is in the hierarchy with the value indicating whether implicitly or explicitly. Unset if the object does not appear in the hierarchy.*/
-		TOptional<EHierarchyObjectType> IsInHierarchy(const FSoftObjectPath& Object) const;
+		UE_API TOptional<EHierarchyObjectType> IsInHierarchy(const FSoftObjectPath& Object) const;
 		/** @return Whether Object has any subobjects in the hierarchy. */
-		bool HasChildren(const FSoftObjectPath& Object) const;
+		UE_API bool HasChildren(const FSoftObjectPath& Object) const;
 		/** @return Whether Object is an asset, i.e. a top-level asset that has never have any parents. */
-		bool IsAssetInHierarchy(const FSoftObjectPath& Object) const;
+		UE_API bool IsAssetInHierarchy(const FSoftObjectPath& Object) const;
 
 		/** Checks there is an hierarchy. */
 		bool IsEmpty() const { return AssetNodes.IsEmpty(); }
@@ -122,7 +124,7 @@ namespace UE::ConcertSyncCore
 		 * Henceforth, the object shall be tracked as EHierarchyObjectType::Explicit.
 		 * Calling this more than once has no effect.
 		 */
-		void AddObject(const FSoftObjectPath& ObjectPath);
+		UE_API void AddObject(const FSoftObjectPath& ObjectPath);
 
 		/**
 		 * Removes Object from the hierarchy.
@@ -130,7 +132,7 @@ namespace UE::ConcertSyncCore
 		 * If the object has no children, the object is removed entirely.
 		 * If the object has children, it shall henceforth be tacked as EHierarchyObjectType::Implicit.
 		 */
-		void RemoveObject(const FSoftObjectPath& Object);
+		UE_API void RemoveObject(const FSoftObjectPath& Object);
 
 		/** Empties the entire hierarchy to the empty state. */
 		void Clear() { AssetNodes.Reset(); CachedNodes.Reset(); }
@@ -159,15 +161,17 @@ namespace UE::ConcertSyncCore
 		/** Maps every path to the corresponding node. */
 		TMap<FSoftObjectPath, FTreeNode*> CachedNodes;
 
-		static ETreeTraversalBehavior TraverseTopToBottomInternal(
+		static UE_API ETreeTraversalBehavior TraverseTopToBottomInternal(
 			const FObjectHierarchyInfo& Owner,
 			const TArray<TUniquePtr<FTreeNode>>& Nodes,
 			TFunctionRef<ETreeTraversalBehavior(const FChildRelation& Relation)> Callback
 			);
-		static EBreakBehavior TraverseBottomToTopInternal(
+		static UE_API EBreakBehavior TraverseBottomToTopInternal(
 			const FObjectHierarchyInfo& Owner,
 			const TArray<TUniquePtr<FTreeNode>>& Nodes,
 			TFunctionRef<EBreakBehavior(const FChildRelation& Relation)> Callback
 			);
 	};
-}
\ No newline at end of file
+}
+
+#undef UE_API
diff --git a/Engine/Plugins/Developer/Concert/ConcertSync/ConcertSyncCore/Source/ConcertSyncCore/Public/Replication/Data/ConcertPropertySelection.h b/Engine/Plugins/Developer/Concert/ConcertSync/ConcertSyncCore/Source/ConcertSyncCore/Public/Replication/Data/ConcertPropertySelection.h
--- a/Engine/Plugins/Developer/Concert/ConcertSync/ConcertSyncCore/Source/ConcertSyncCore/Public/Replication/Data/ConcertPropertySelection.h
+++ b/Engine/Plugins/Developer/Concert/ConcertSync/ConcertSyncCore/Source/ConcertSyncCore/Public/Replication/Data/ConcertPropertySelection.h
@@ -5,6 +5,8 @@
 #include "Misc/EBreakBehavior.h"
 #include "ConcertPropertySelection.generated.h"
 
+#define UE_API CONCERTSYNCCORE_API
+
 class UStruct;
 struct FArchiveSerializedPropertyChain;
 
@@ -13,19 +15,19 @@ struct FArchiveSerializedPropertyChain;
  * @see FConcertPropertyChain::PathToProperty.
  */
 USTRUCT()
-struct CONCERTSYNCCORE_API FConcertPropertyChain
+struct FConcertPropertyChain
 {
 	GENERATED_BODY()
 
 	/** Constructs a FConcertPropertyChain from a path if it is valid. If you need to create many paths in one go, use PropertyUtils::BulkConstructConcertChainsFromPaths instead. */
-	static TOptional<FConcertPropertyChain> CreateFromPath(const UStruct& Class, const TArray<FName>& NamePath);
+	static UE_API TOptional<FConcertPropertyChain> CreateFromPath(const UStruct& Class, const TArray<FName>& NamePath);
 
 	FConcertPropertyChain() = default;
 	/**
 	 * @param OptionalChain The chain leading up to LeafProperty. If it is a root property it can either be empty or nullptr. This mimics the behaviours of FArchive.
 	 * @param LeafProperty The property the chain leads to. It will be the last property in PathToProperty. This can be the inner property of a container but ONLY for primitive types (float, etc.) or structs with a custom Serialize function. 
 	 */
-	FConcertPropertyChain(const FArchiveSerializedPropertyChain* OptionalChain, const FProperty& LeafProperty);
+	UE_API FConcertPropertyChain(const FArchiveSerializedPropertyChain* OptionalChain, const FProperty& LeafProperty);
 	
 	/** Gets the leaf property, which is the property the path leads towards. */
 	FName GetLeafProperty() const { return IsEmpty() ? NAME_None : PathToProperty[PathToProperty.Num() - 1]; }
@@ -37,9 +39,9 @@ struct CONCERTSYNCCORE_API FConcertPropertyChain
 	bool IsParentOf(const FConcertPropertyChain& ChildToCheck) const { return ChildToCheck.IsChildOf(*this); }
 	
 	/** @return Whether the leaf property is a child of the given property chain. */
-	bool IsChildOf(const FConcertPropertyChain& ParentToCheck) const;
+	UE_API bool IsChildOf(const FConcertPropertyChain& ParentToCheck) const;
 	/** @return Whether the leaf property is a direct child of the given property chain. */
-	bool IsDirectChildOf(const FConcertPropertyChain& ParentToCheck) const;
+	UE_API bool IsDirectChildOf(const FConcertPropertyChain& ParentToCheck) const;
 
 	/**
 	 * Utility for checking whether this path corresponds to OptionalChain leading to LeafProperty.
@@ -47,15 +49,15 @@ struct CONCERTSYNCCORE_API FConcertPropertyChain
 	 * @param OptionalChain Contains all properties leading up to the LeafProperty. You can skip inner container properties.
 	 * @param LeafProperty The last property
 	 * @return Whether OptionalChain and LeafProperty correspond to this path. */
-	bool MatchesExactly(const FArchiveSerializedPropertyChain* OptionalChain, const FProperty& LeafProperty) const;
+	UE_API bool MatchesExactly(const FArchiveSerializedPropertyChain* OptionalChain, const FProperty& LeafProperty) const;
 
 	/** @return Attempts to resolve this property given the class */
-	FProperty* ResolveProperty(const UStruct& Class, bool bLogOnFail = true) const;
+	UE_API FProperty* ResolveProperty(const UStruct& Class, bool bLogOnFail = true) const;
 
 	/** @return The property immediately before the current one in the chain */
-	FConcertPropertyChain GetParent() const;
+	UE_API FConcertPropertyChain GetParent() const;
 	/** @return The root-most parent in the chain. */
-	FConcertPropertyChain GetRootParent() const;
+	UE_API FConcertPropertyChain GetRootParent() const;
 	/** @return The property path */
 	const TArray<FName>& GetPathToProperty() const { return PathToProperty; }
 
@@ -64,7 +66,7 @@ struct CONCERTSYNCCORE_API FConcertPropertyChain
 		Path,
 		LeafProperty
 	};
-	FString ToString(EToStringMethod Method = EToStringMethod::Path) const;
+	UE_API FString ToString(EToStringMethod Method = EToStringMethod::Path) const;
 	
 	friend bool operator==(const FConcertPropertyChain& Left, const FConcertPropertyChain& Right) { return Left.PathToProperty == Right.PathToProperty; }
 	friend bool operator!=(const FConcertPropertyChain& Left, const FConcertPropertyChain& Right) { return !(Left == Right); }
@@ -128,7 +130,7 @@ private:
 
 /** List of properties to be replicated for a given object */
 USTRUCT()
-struct CONCERTSYNCCORE_API FConcertPropertySelection
+struct FConcertPropertySelection
 {
 	GENERATED_BODY()
 
@@ -148,7 +150,7 @@ struct CONCERTSYNCCORE_API FConcertPropertySelection
 	 * Suppose ReplicatedProperties = { ["Vector", "X"] }.
 	 * After execution, it would be { ["Vector", "X"], ["Vector"] }
 	 */
-	void DiscoverAndAddImplicitParentProperties();
+	UE_API void DiscoverAndAddImplicitParentProperties();
 	
 	/**
 	 * Determines all properties that overlap.
@@ -156,7 +158,7 @@ struct CONCERTSYNCCORE_API FConcertPropertySelection
 	 * 
 	 * @return Whether there were any property overlaps.
 	 */
-	static bool EnumeratePropertyOverlaps(
+	static UE_API bool EnumeratePropertyOverlaps(
 		const TSet<FConcertPropertyChain>& First,
 		const TSet<FConcertPropertyChain>& Second,
 		TFunctionRef<EBreakBehavior(const FConcertPropertyChain&)> Callback = [](const FConcertPropertyChain&){ return EBreakBehavior::Break; }
@@ -186,3 +188,4 @@ inline uint32 GetTypeHash(const FConcertPropertyChain& Chain)
 	return UE::ConcertSyncCore::ComputeHashForPropertyChainContent(Chain.GetPathToProperty());
 }
 
+#undef UE_API
diff --git a/Engine/Plugins/Developer/Concert/ConcertSync/ConcertSyncCore/Source/ConcertSyncCore/Public/Replication/Data/ReplicationFrequencySettings.h b/Engine/Plugins/Developer/Concert/ConcertSync/ConcertSyncCore/Source/ConcertSyncCore/Public/Replication/Data/ReplicationFrequencySettings.h
--- a/Engine/Plugins/Developer/Concert/ConcertSync/ConcertSyncCore/Source/ConcertSyncCore/Public/Replication/Data/ReplicationFrequencySettings.h
+++ b/Engine/Plugins/Developer/Concert/ConcertSync/ConcertSyncCore/Source/ConcertSyncCore/Public/Replication/Data/ReplicationFrequencySettings.h
@@ -32,7 +32,7 @@ namespace UE::ConcertSyncCore
 
 /** Specifies replication frequency settings for a single object. */
 USTRUCT()
-struct CONCERTSYNCCORE_API FConcertObjectReplicationSettings
+struct FConcertObjectReplicationSettings
 {
 	GENERATED_BODY()
 	
@@ -116,7 +116,7 @@ struct CONCERTSYNCCORE_API FConcertObjectReplicationSettings
 
 /** Specifies replication frequency settings for an entire stream. */
 USTRUCT()
-struct CONCERTSYNCCORE_API FConcertStreamFrequencySettings
+struct FConcertStreamFrequencySettings
 {
 	GENERATED_BODY()
 
@@ -147,4 +147,4 @@ struct CONCERTSYNCCORE_API FConcertStreamFrequencySettings
 	{
 		return !(Left == Right);
 	}
-};
\ No newline at end of file
+};
diff --git a/Engine/Plugins/Developer/Concert/ConcertSync/ConcertSyncCore/Source/ConcertSyncCore/Public/Replication/Formats/FullObjectFormat.h b/Engine/Plugins/Developer/Concert/ConcertSync/ConcertSyncCore/Source/ConcertSyncCore/Public/Replication/Formats/FullObjectFormat.h
--- a/Engine/Plugins/Developer/Concert/ConcertSync/ConcertSyncCore/Source/ConcertSyncCore/Public/Replication/Formats/FullObjectFormat.h
+++ b/Engine/Plugins/Developer/Concert/ConcertSync/ConcertSyncCore/Source/ConcertSyncCore/Public/Replication/Formats/FullObjectFormat.h
@@ -4,6 +4,8 @@
 
 #include "IObjectReplicationFormat.h"
 
+#define UE_API CONCERTSYNCCORE_API
+
 namespace UE::ConcertSyncCore
 {
 	/**
@@ -17,15 +19,17 @@ namespace UE::ConcertSyncCore
 	 *
 	 * @see FFullObjectReplicationData
 	 */
-	class CONCERTSYNCCORE_API FFullObjectFormat : public IObjectReplicationFormat
+	class FFullObjectFormat : public IObjectReplicationFormat
 	{
 	public:
 
 		//~ Begin FFullObjectFormat Interface
-		virtual TOptional<FConcertSessionSerializedPayload> CreateReplicationEvent(UObject& Object, FAllowPropertyFunc IsPropertyAllowedFunc) override;
+		UE_API virtual TOptional<FConcertSessionSerializedPayload> CreateReplicationEvent(UObject& Object, FAllowPropertyFunc IsPropertyAllowedFunc) override;
 		virtual void ClearInternalCache(TArrayView<UObject> ObjectsToClear) override { /* There is no cache: Full object format has no smart building relying on past events - it always serializes everything */ }
-		virtual void CombineReplicationEvents(FConcertSessionSerializedPayload& Base, const FConcertSessionSerializedPayload& Newer) override;
-		virtual void ApplyReplicationEvent(UObject& Object, const FConcertSessionSerializedPayload& Payload, const FOnPropertyVisitedFunc& OnPrePropertySerialized) override;
+		UE_API virtual void CombineReplicationEvents(FConcertSessionSerializedPayload& Base, const FConcertSessionSerializedPayload& Newer) override;
+		UE_API virtual void ApplyReplicationEvent(UObject& Object, const FConcertSessionSerializedPayload& Payload, const FOnPropertyVisitedFunc& OnPrePropertySerialized) override;
 		//~ End FFullObjectFormat Interface
 	};
 }
+
+#undef UE_API
diff --git a/Engine/Plugins/Developer/Concert/ConcertSync/ConcertSyncCore/Source/ConcertSyncCore/Public/Replication/Formats/IObjectReplicationFormat.h b/Engine/Plugins/Developer/Concert/ConcertSync/ConcertSyncCore/Source/ConcertSyncCore/Public/Replication/Formats/IObjectReplicationFormat.h
--- a/Engine/Plugins/Developer/Concert/ConcertSync/ConcertSyncCore/Source/ConcertSyncCore/Public/Replication/Formats/IObjectReplicationFormat.h
+++ b/Engine/Plugins/Developer/Concert/ConcertSync/ConcertSyncCore/Source/ConcertSyncCore/Public/Replication/Formats/IObjectReplicationFormat.h
@@ -21,7 +21,7 @@ namespace UE::ConcertSyncCore
 	 * - full-object format that serializes all replicated properties and never does any combining (inefficient).
 	 * - partial-object format that keeps track of past properties and only sends changed replicated properties.
 	 */
-	class CONCERTSYNCCORE_API IObjectReplicationFormat
+	class IObjectReplicationFormat
 	{
 	public:
 
@@ -72,4 +72,4 @@ namespace UE::ConcertSyncCore
 
 		virtual ~IObjectReplicationFormat() = default;
 	};
-}
\ No newline at end of file
+}
diff --git a/Engine/Plugins/Developer/Concert/ConcertSync/ConcertSyncCore/Source/ConcertSyncCore/Public/Replication/Messages/ReplicationActivity.h b/Engine/Plugins/Developer/Concert/ConcertSync/ConcertSyncCore/Source/ConcertSyncCore/Public/Replication/Messages/ReplicationActivity.h
--- a/Engine/Plugins/Developer/Concert/ConcertSync/ConcertSyncCore/Source/ConcertSyncCore/Public/Replication/Messages/ReplicationActivity.h
+++ b/Engine/Plugins/Developer/Concert/ConcertSync/ConcertSyncCore/Source/ConcertSyncCore/Public/Replication/Messages/ReplicationActivity.h
@@ -1,4 +1,4 @@
-﻿// Copyright Epic Games, Inc. All Rights Reserved.
+// Copyright Epic Games, Inc. All Rights Reserved.
 
 #pragma once
 
@@ -10,6 +10,8 @@
 #include "Replication/Data/ReplicationStream.h"
 #include "ReplicationActivity.generated.h"
 
+#define UE_API CONCERTSYNCCORE_API
+
 class FText;
 
 /** Identifies the FConcertSyncReplicationEvent::Payload struct type */
@@ -187,7 +189,7 @@ struct FConcertSyncReplicationActivity : public FConcertSyncActivity
 
 /** Summary for a lock activity entry in a Concert Sync Session */
 USTRUCT()
-struct CONCERTSYNCCORE_API FConcertSyncReplicationActivitySummary : public FConcertSyncActivitySummary
+struct FConcertSyncReplicationActivitySummary : public FConcertSyncActivitySummary
 {
 	GENERATED_BODY()
 
@@ -211,15 +213,17 @@ struct CONCERTSYNCCORE_API FConcertSyncReplicationActivitySummary : public FConc
 	}
 
 	/** Create this summary from a replication event */
-	static FConcertSyncReplicationActivitySummary CreateSummaryForEvent(const FConcertSyncReplicationEvent& InEvent);
+	static UE_API FConcertSyncReplicationActivitySummary CreateSummaryForEvent(const FConcertSyncReplicationEvent& InEvent);
 
 	/** Gets the title for this summary */
-	FText ToDisplayTitle() const;
+	UE_API FText ToDisplayTitle() const;
 
 protected:
 	
 	//~ Begin FConcertSyncActivitySummary Interface
-	virtual FText CreateDisplayText(const bool InUseRichText) const override;
-	virtual FText CreateDisplayTextForUser(const FText InUserDisplayName, const bool InUseRichText) const override;
+	UE_API virtual FText CreateDisplayText(const bool InUseRichText) const override;
+	UE_API virtual FText CreateDisplayTextForUser(const FText InUserDisplayName, const bool InUseRichText) const override;
 	//~ End FConcertSyncActivitySummary Interface
-};
\ No newline at end of file
+};
+
+#undef UE_API
diff --git a/Engine/Plugins/Developer/Concert/ConcertSync/ConcertSyncCore/Source/ConcertSyncCore/Public/Replication/Misc/ReplicatedObjectHierarchyCache.h b/Engine/Plugins/Developer/Concert/ConcertSync/ConcertSyncCore/Source/ConcertSyncCore/Public/Replication/Misc/ReplicatedObjectHierarchyCache.h
--- a/Engine/Plugins/Developer/Concert/ConcertSync/ConcertSyncCore/Source/ConcertSyncCore/Public/Replication/Misc/ReplicatedObjectHierarchyCache.h
+++ b/Engine/Plugins/Developer/Concert/ConcertSync/ConcertSyncCore/Source/ConcertSyncCore/Public/Replication/Misc/ReplicatedObjectHierarchyCache.h
@@ -1,4 +1,4 @@
-﻿// Copyright Epic Games, Inc. All Rights Reserved.
+// Copyright Epic Games, Inc. All Rights Reserved.
 
 #pragma once
 
@@ -7,6 +7,8 @@
 #include "Containers/Array.h"
 #include "Containers/ContainersFwd.h"
 
+#define UE_API CONCERTSYNCCORE_API
+
 struct FConcertObjectInStreamID;
 struct FConcertReplicatedObjectId;
 struct FConcertReplication_ChangeStream_Request;
@@ -19,7 +21,7 @@ namespace UE::ConcertSyncCore
 	 * Holds the outer hierarchy of all objects registered in any stream.
 	 * Updates the hierarchy when a relevant event changing stream content happens.
 	 */
-	class CONCERTSYNCCORE_API FReplicatedObjectHierarchyCache : FObjectPathHierarchy
+	class FReplicatedObjectHierarchyCache : FObjectPathHierarchy
 	{
 	public:
 
@@ -31,7 +33,7 @@ namespace UE::ConcertSyncCore
 		using FObjectPathHierarchy::IsEmpty;
 
 		/** Removes all saved state. */
-		void Clear();
+		UE_API void Clear();
 
 		/**
 		 * Checks whether the object is registered by any client (except for this in IgnoredClients).
@@ -42,12 +44,12 @@ namespace UE::ConcertSyncCore
 		 * - IsObjectReferencedDirectly(/Game/Maps.Map:PersistentLevel.Cube) == false
 		 * but e.g. TraverseTopToBottom would list both paths.
 		 */
-		bool IsObjectReferencedDirectly(const FSoftObjectPath& ObjectPath, TConstArrayView<FGuid> IgnoredClients = {}) const;
+		UE_API bool IsObjectReferencedDirectly(const FSoftObjectPath& ObjectPath, TConstArrayView<FGuid> IgnoredClients = {}) const;
 
 		// Validated network events that modify the known hierarchy of objects.
-		void OnJoin(const FGuid& ClientId, const FConcertReplication_Join_Request& Request);
-		void OnPostClientLeft(const FGuid& ClientId, TConstArrayView<FConcertReplicationStream> Streams);
-		void OnChangeStreams(const FGuid& ClientId, TConstArrayView<FConcertObjectInStreamID> AddedObjects, TConstArrayView<FConcertObjectInStreamID> RemovedObjects);
+		UE_API void OnJoin(const FGuid& ClientId, const FConcertReplication_Join_Request& Request);
+		UE_API void OnPostClientLeft(const FGuid& ClientId, TConstArrayView<FConcertReplicationStream> Streams);
+		UE_API void OnChangeStreams(const FGuid& ClientId, TConstArrayView<FConcertObjectInStreamID> AddedObjects, TConstArrayView<FConcertObjectInStreamID> RemovedObjects);
 
 	private:
 
@@ -68,8 +70,10 @@ namespace UE::ConcertSyncCore
 		
 		TMap<FSoftObjectPath, FObjectMetaData> ObjectMetaData;
 
-		void AddObjectInternal(const FConcertReplicatedObjectId& ObjectInfo);
-		void RemoveObjectInternal(const FConcertReplicatedObjectId& Object);
-		bool RemoveMetaData(const FConcertReplicatedObjectId& Object);
+		UE_API void AddObjectInternal(const FConcertReplicatedObjectId& ObjectInfo);
+		UE_API void RemoveObjectInternal(const FConcertReplicatedObjectId& Object);
+		UE_API bool RemoveMetaData(const FConcertReplicatedObjectId& Object);
 	};
 }
+
+#undef UE_API
diff --git a/Engine/Plugins/Developer/Concert/ConcertSync/ConcertSyncCore/Source/ConcertSyncCore/Public/Replication/Processing/Actions/ConcertReplicationAction.h b/Engine/Plugins/Developer/Concert/ConcertSync/ConcertSyncCore/Source/ConcertSyncCore/Public/Replication/Processing/Actions/ConcertReplicationAction.h
--- a/Engine/Plugins/Developer/Concert/ConcertSync/ConcertSyncCore/Source/ConcertSyncCore/Public/Replication/Processing/Actions/ConcertReplicationAction.h
+++ b/Engine/Plugins/Developer/Concert/ConcertSync/ConcertSyncCore/Source/ConcertSyncCore/Public/Replication/Processing/Actions/ConcertReplicationAction.h
@@ -1,4 +1,4 @@
-﻿// Copyright Epic Games, Inc. All Rights Reserved.
+// Copyright Epic Games, Inc. All Rights Reserved.
 
 #pragma once
 
@@ -32,10 +32,10 @@ struct FReplicationActionArgs
  * i.e. UPROPERTY(Config) TInstancedStruct<FConcertReplicationAction>.
  */
 USTRUCT()
-struct CONCERTSYNCCORE_API FConcertReplicationAction
+struct FConcertReplicationAction
 {
 	GENERATED_BODY()
 
 	virtual ~FConcertReplicationAction() = default;
 	virtual void Apply(const UE::ConcertSyncCore::FReplicationActionArgs& InArgs) const { checkNoEntry(); }
-};
\ No newline at end of file
+};
diff --git a/Engine/Plugins/Developer/Concert/ConcertSync/ConcertSyncCore/Source/ConcertSyncCore/Public/Replication/Processing/Actions/ConcertReplicationAction_MarkRenderStateDirty.h b/Engine/Plugins/Developer/Concert/ConcertSync/ConcertSyncCore/Source/ConcertSyncCore/Public/Replication/Processing/Actions/ConcertReplicationAction_MarkRenderStateDirty.h
--- a/Engine/Plugins/Developer/Concert/ConcertSync/ConcertSyncCore/Source/ConcertSyncCore/Public/Replication/Processing/Actions/ConcertReplicationAction_MarkRenderStateDirty.h
+++ b/Engine/Plugins/Developer/Concert/ConcertSync/ConcertSyncCore/Source/ConcertSyncCore/Public/Replication/Processing/Actions/ConcertReplicationAction_MarkRenderStateDirty.h
@@ -10,7 +10,7 @@
 
 /** Calls MarkRenderStateDirty on the object if it is a USceneComponent. */
 USTRUCT(DisplayName = "Mark Render State Dirty")
-struct CONCERTSYNCCORE_API FConcertReplicationAction_MarkRenderStateDirty : public FConcertReplicationAction
+struct FConcertReplicationAction_MarkRenderStateDirty : public FConcertReplicationAction
 {
 	GENERATED_BODY()
 
@@ -23,4 +23,4 @@ struct CONCERTSYNCCORE_API FConcertReplicationAction_MarkRenderStateDirty : publ
 		}
 #endif
 	}
-};
\ No newline at end of file
+};
diff --git a/Engine/Plugins/Developer/Concert/ConcertSync/ConcertSyncCore/Source/ConcertSyncCore/Public/Replication/Processing/Actions/ConcertReplicationAction_PostEditChange.h b/Engine/Plugins/Developer/Concert/ConcertSync/ConcertSyncCore/Source/ConcertSyncCore/Public/Replication/Processing/Actions/ConcertReplicationAction_PostEditChange.h
--- a/Engine/Plugins/Developer/Concert/ConcertSync/ConcertSyncCore/Source/ConcertSyncCore/Public/Replication/Processing/Actions/ConcertReplicationAction_PostEditChange.h
+++ b/Engine/Plugins/Developer/Concert/ConcertSync/ConcertSyncCore/Source/ConcertSyncCore/Public/Replication/Processing/Actions/ConcertReplicationAction_PostEditChange.h
@@ -8,7 +8,7 @@
 
 /** Calls PostEditChange on the object. */
 USTRUCT(DisplayName = "Post Edit Change")
-struct CONCERTSYNCCORE_API FConcertReplicationAction_PostEditChange : public FConcertReplicationAction
+struct FConcertReplicationAction_PostEditChange : public FConcertReplicationAction
 {
 	GENERATED_BODY()
 
@@ -18,4 +18,4 @@ struct CONCERTSYNCCORE_API FConcertReplicationAction_PostEditChange : public FCo
 		InArgs.Object->PostEditChange();
 #endif
 	}
-};
\ No newline at end of file
+};
diff --git a/Engine/Plugins/Developer/Concert/ConcertSync/ConcertSyncCore/Source/ConcertSyncCore/Public/Replication/Processing/Actions/ConcertReplicationAction_UpdateComponentToWorld.h b/Engine/Plugins/Developer/Concert/ConcertSync/ConcertSyncCore/Source/ConcertSyncCore/Public/Replication/Processing/Actions/ConcertReplicationAction_UpdateComponentToWorld.h
--- a/Engine/Plugins/Developer/Concert/ConcertSync/ConcertSyncCore/Source/ConcertSyncCore/Public/Replication/Processing/Actions/ConcertReplicationAction_UpdateComponentToWorld.h
+++ b/Engine/Plugins/Developer/Concert/ConcertSync/ConcertSyncCore/Source/ConcertSyncCore/Public/Replication/Processing/Actions/ConcertReplicationAction_UpdateComponentToWorld.h
@@ -10,7 +10,7 @@
 
 /** Calls ConcertReplicationAction_UpdateComponentToWorld on the object if it is a USceneComponent. */
 USTRUCT(DisplayName = "Update Component To World")
-struct CONCERTSYNCCORE_API FConcertReplicationAction_UpdateComponentToWorld : public FConcertReplicationAction
+struct FConcertReplicationAction_UpdateComponentToWorld : public FConcertReplicationAction
 {
 	GENERATED_BODY()
 
@@ -23,4 +23,4 @@ struct CONCERTSYNCCORE_API FConcertReplicationAction_UpdateComponentToWorld : pu
 		}
 #endif
 	}
-};
\ No newline at end of file
+};
diff --git a/Engine/Plugins/Developer/Concert/ConcertSync/ConcertSyncCore/Source/ConcertSyncCore/Public/Replication/Processing/Actions/ReplicationActionDispatcher.h b/Engine/Plugins/Developer/Concert/ConcertSync/ConcertSyncCore/Source/ConcertSyncCore/Public/Replication/Processing/Actions/ReplicationActionDispatcher.h
--- a/Engine/Plugins/Developer/Concert/ConcertSync/ConcertSyncCore/Source/ConcertSyncCore/Public/Replication/Processing/Actions/ReplicationActionDispatcher.h
+++ b/Engine/Plugins/Developer/Concert/ConcertSync/ConcertSyncCore/Source/ConcertSyncCore/Public/Replication/Processing/Actions/ReplicationActionDispatcher.h
@@ -1,9 +1,11 @@
-﻿// Copyright Epic Games, Inc. All Rights Reserved.
+// Copyright Epic Games, Inc. All Rights Reserved.
 
 #pragma once
 
 #include "Containers/BitArray.h"
 
+#define UE_API CONCERTSYNCCORE_API
+
 struct FConcertReplicationActionEntry;
 
 namespace UE::ConcertSyncCore
@@ -11,17 +13,17 @@ namespace UE::ConcertSyncCore
 struct FReplicationActionArgs;
 
 /** Util that you can use when calling IObjectReplicationFormat::ApplyReplicationEvent to execute actions. */
-class CONCERTSYNCCORE_API FReplicationActionDispatcher
+class FReplicationActionDispatcher
 {
 public:
 
-	explicit FReplicationActionDispatcher(const TArray<FConcertReplicationActionEntry>& InActions UE_LIFETIMEBOUND, bool bDebugActions = false);
+	UE_API explicit FReplicationActionDispatcher(const TArray<FConcertReplicationActionEntry>& InActions UE_LIFETIMEBOUND, bool bDebugActions = false);
 
 	/** Checks whether InProperty should trigger actions. */
-	void OnReplicateProperty(const FProperty& InProperty);
+	UE_API void OnReplicateProperty(const FProperty& InProperty);
 
 	/** Call all properties have been processed using OnReplicateProperty. Triggers the actions that need to triggered.  */
-	void ExecuteActions(const FReplicationActionArgs& InArgs);
+	UE_API void ExecuteActions(const FReplicationActionArgs& InArgs);
 
 private:
 
@@ -35,3 +37,5 @@ private:
 	TBitArray<> ActionsToTrigger;
 };
 }
+
+#undef UE_API
diff --git a/Engine/Plugins/Developer/Concert/ConcertSync/ConcertSyncCore/Source/ConcertSyncCore/Public/Replication/Processing/ObjectReplicationCache.h b/Engine/Plugins/Developer/Concert/ConcertSync/ConcertSyncCore/Source/ConcertSyncCore/Public/Replication/Processing/ObjectReplicationCache.h
--- a/Engine/Plugins/Developer/Concert/ConcertSync/ConcertSyncCore/Source/ConcertSyncCore/Public/Replication/Processing/ObjectReplicationCache.h
+++ b/Engine/Plugins/Developer/Concert/ConcertSync/ConcertSyncCore/Source/ConcertSyncCore/Public/Replication/Processing/ObjectReplicationCache.h
@@ -10,6 +10,8 @@
 #include "Templates/Function.h"
 #include "Templates/SharedPointer.h"
 
+#define UE_API CONCERTSYNCCORE_API
+
 struct FConcertReplication_ObjectReplicationEvent;
 struct FSoftObjectPath;
 struct FGuid;
@@ -22,7 +24,7 @@ namespace UE::ConcertSyncCore
 	 * Interface for objects that want to latently use received FConcertObjectReplicationEvent data.
 	 * The data is continuously updated until it is consumed.
 	 */
-	class CONCERTSYNCCORE_API IReplicationCacheUser
+	class IReplicationCacheUser
 	{
 	public:
 
@@ -67,11 +69,11 @@ namespace UE::ConcertSyncCore
 	 * This allows multiple systems to reuse replication data. For example, on the server the same data may need to be distributed to different clients but
 	 * the clients send the data at different times.
 	 */
-	class CONCERTSYNCCORE_API FObjectReplicationCache : public TSharedFromThis<FObjectReplicationCache>
+	class FObjectReplicationCache : public TSharedFromThis<FObjectReplicationCache>
 	{
 	public:
 
-		FObjectReplicationCache(IObjectReplicationFormat& ReplicationFormat UE_LIFETIMEBOUND);
+		UE_API FObjectReplicationCache(IObjectReplicationFormat& ReplicationFormat UE_LIFETIMEBOUND);
 		
 		/**
 		 * Called when new data is received for an object and shares it with any IObjectCacheUser that is possibly interested in it.
@@ -81,11 +83,11 @@ namespace UE::ConcertSyncCore
 		 * @param ObjectReplicationEvent The data that was replicated
 		 * @return The number of cache users that accepted this event. 
 		 */
-		FCacheStoreStats StoreUntilConsumed(const FGuid& SendingEndpointId, const FGuid& OriginStreamId, const FSequenceId SequenceId, const FConcertReplication_ObjectReplicationEvent& ObjectReplicationEvent);
+		UE_API FCacheStoreStats StoreUntilConsumed(const FGuid& SendingEndpointId, const FGuid& OriginStreamId, const FSequenceId SequenceId, const FConcertReplication_ObjectReplicationEvent& ObjectReplicationEvent);
 
 		/** Registers a new user, which will start receiving data for any new data received from now on. */
-		void RegisterDataCacheUser(TSharedRef<IReplicationCacheUser> User);
-		void UnregisterDataCacheUser(const TSharedRef<IReplicationCacheUser>& User);
+		UE_API void RegisterDataCacheUser(TSharedRef<IReplicationCacheUser> User);
+		UE_API void UnregisterDataCacheUser(const TSharedRef<IReplicationCacheUser>& User);
 
 	private:
 
@@ -116,11 +118,13 @@ namespace UE::ConcertSyncCore
 		 * Combines old data that cache users may have associated for ObjectId with the NewData and notifies the users with OnCachedDataUpdated.
 		 * @return Number of cache users that were updated
 		 */
-		uint32 CombineCachedDataWithNewData(
+		UE_API uint32 CombineCachedDataWithNewData(
 			const FConcertReplicatedObjectId& ObjectId,
 			FSequenceId NewSequenceId,
 			const FConcertReplication_ObjectReplicationEvent& NewData,
 			const FObjectCache& ObjectCacheBeforeAddition
 			) const;
 	};
-}
\ No newline at end of file
+}
+
+#undef UE_API
diff --git a/Engine/Plugins/Developer/Concert/ConcertSync/ConcertSyncCore/Source/ConcertSyncCore/Public/Replication/Processing/ObjectReplicationProcessor.h b/Engine/Plugins/Developer/Concert/ConcertSync/ConcertSyncCore/Source/ConcertSyncCore/Public/Replication/Processing/ObjectReplicationProcessor.h
--- a/Engine/Plugins/Developer/Concert/ConcertSync/ConcertSyncCore/Source/ConcertSyncCore/Public/Replication/Processing/ObjectReplicationProcessor.h
+++ b/Engine/Plugins/Developer/Concert/ConcertSync/ConcertSyncCore/Source/ConcertSyncCore/Public/Replication/Processing/ObjectReplicationProcessor.h
@@ -8,6 +8,8 @@
 #include "Templates/SharedPointer.h"
 #include "UObject/SoftObjectPath.h"
 
+#define UE_API CONCERTSYNCCORE_API
+
 struct FConcertObjectReplicationSettings;
 struct FConcertObjectInStreamID;
 
@@ -36,17 +38,17 @@ namespace UE::ConcertSyncCore
 	 *  
 	 * FObjectReplicationProcessor is backed by an IReplicationDataSource, which is a source of replicated object data.
 	 */
-	class CONCERTSYNCCORE_API FObjectReplicationProcessor
+	class FObjectReplicationProcessor
 	{
 	public:
 		
 		/**
 		 * @param DataSource Source of the data that is to be sent
 		 */
-		FObjectReplicationProcessor(IReplicationDataSource& DataSource UE_LIFETIMEBOUND);
+		UE_API FObjectReplicationProcessor(IReplicationDataSource& DataSource UE_LIFETIMEBOUND);
 		virtual ~FObjectReplicationProcessor() = default;
 		
-		virtual void ProcessObjects(const FProcessObjectsParams& Params);
+		UE_API virtual void ProcessObjects(const FProcessObjectsParams& Params);
 		
 	protected:
 
@@ -67,3 +69,5 @@ namespace UE::ConcertSyncCore
 		IReplicationDataSource& DataSource;
 	};
 }
+
+#undef UE_API
diff --git a/Engine/Plugins/Developer/Concert/ConcertSync/ConcertSyncCore/Source/ConcertSyncCore/Public/Replication/Processing/ObjectReplicationReceiver.h b/Engine/Plugins/Developer/Concert/ConcertSync/ConcertSyncCore/Source/ConcertSyncCore/Public/Replication/Processing/ObjectReplicationReceiver.h
--- a/Engine/Plugins/Developer/Concert/ConcertSync/ConcertSyncCore/Source/ConcertSyncCore/Public/Replication/Processing/ObjectReplicationReceiver.h
+++ b/Engine/Plugins/Developer/Concert/ConcertSync/ConcertSyncCore/Source/ConcertSyncCore/Public/Replication/Processing/ObjectReplicationReceiver.h
@@ -4,6 +4,8 @@
 
 #include "HAL/Platform.h"
 
+#define UE_API CONCERTSYNCCORE_API
+
 class IConcertSession;
 struct FConcertReplication_BatchReplicationEvent;
 struct FConcertReplication_ObjectReplicationEvent;
@@ -15,12 +17,12 @@ namespace UE::ConcertSyncCore
 	class FObjectReplicationCache;
 	
 	/** Receives replicated object data from all message endpoints and stores it in a FObjectReplicationCache. */
-	class CONCERTSYNCCORE_API FObjectReplicationReceiver
+	class FObjectReplicationReceiver
 	{
 	public:
 
-		FObjectReplicationReceiver(IConcertSession& Session UE_LIFETIMEBOUND, FObjectReplicationCache& ReplicationCache UE_LIFETIMEBOUND);
-		virtual ~FObjectReplicationReceiver();
+		UE_API FObjectReplicationReceiver(IConcertSession& Session UE_LIFETIMEBOUND, FObjectReplicationCache& ReplicationCache UE_LIFETIMEBOUND);
+		UE_API virtual ~FObjectReplicationReceiver();
 
 	protected:
 
@@ -34,6 +36,8 @@ namespace UE::ConcertSyncCore
 		/** Where received data is stored. */
 		FObjectReplicationCache& ReplicationCache;
 
-		void HandleBatchReplicationEvent(const FConcertSessionContext& SessionContext, const FConcertReplication_BatchReplicationEvent& Event);
+		UE_API void HandleBatchReplicationEvent(const FConcertSessionContext& SessionContext, const FConcertReplication_BatchReplicationEvent& Event);
 	};
 }
+
+#undef UE_API
diff --git a/Engine/Plugins/Developer/Concert/ConcertSync/ConcertSyncCore/Source/ConcertSyncCore/Public/Replication/Processing/ObjectReplicationSender.h b/Engine/Plugins/Developer/Concert/ConcertSync/ConcertSyncCore/Source/ConcertSyncCore/Public/Replication/Processing/ObjectReplicationSender.h
--- a/Engine/Plugins/Developer/Concert/ConcertSync/ConcertSyncCore/Source/ConcertSyncCore/Public/Replication/Processing/ObjectReplicationSender.h
+++ b/Engine/Plugins/Developer/Concert/ConcertSync/ConcertSyncCore/Source/ConcertSyncCore/Public/Replication/Processing/ObjectReplicationSender.h
@@ -6,6 +6,8 @@
 #include "Replication/Messages/ObjectReplication.h"
 #include "Trace/ConcertTraceConfig.h"
 
+#define UE_API CONCERTSYNCCORE_API
+
 class IConcertSession;
 
 namespace UE::ConcertSyncCore
@@ -13,7 +15,7 @@ namespace UE::ConcertSyncCore
 	/**
 	 * Sends object data to a specific endpoint ID.
 	 */
-	class CONCERTSYNCCORE_API FObjectReplicationSender : public FObjectReplicationProcessor
+	class FObjectReplicationSender : public FObjectReplicationProcessor
 	{
 	public:
 		
@@ -22,20 +24,20 @@ namespace UE::ConcertSyncCore
 		 * @param Session The session to use for sending
 		 * @param DataSource Source of the data that is to be sent
 		 */
-		FObjectReplicationSender(
+		UE_API FObjectReplicationSender(
 			const FGuid& TargetEndpointId,
 			IConcertSession& Session UE_LIFETIMEBOUND, 
 			IReplicationDataSource& DataSource UE_LIFETIMEBOUND
 			);
 		
 		//~ Begin FObjectReplicationProcessor Interface
-		virtual void ProcessObjects(const FProcessObjectsParams& Params) override;
+		UE_API virtual void ProcessObjects(const FProcessObjectsParams& Params) override;
 		//~ End FObjectReplicationProcessor Interface
 
 	protected:
 
 		//~ Begin FObjectReplicationProcessor Interface
-		virtual void ProcessObject(const FObjectProcessArgs& Args) override;
+		UE_API virtual void ProcessObject(const FObjectProcessArgs& Args) override;
 		//~ End FObjectReplicationProcessor Interface
 
 	private:
@@ -54,7 +56,9 @@ namespace UE::ConcertSyncCore
 		/** This event is filled in ProcessObjects and finally sent to TargetEndpointId. */
 		FConcertReplication_BatchReplicationEvent EventToSend;
 
-		void MarkObjectForTrace(const FConcertReplicatedObjectId& Object, FSequenceId Id);
-		void TraceStartSendingMarkedObjects();
+		UE_API void MarkObjectForTrace(const FConcertReplicatedObjectId& Object, FSequenceId Id);
+		UE_API void TraceStartSendingMarkedObjects();
 	};
 }
+
+#undef UE_API
diff --git a/Engine/Plugins/Developer/Concert/ConcertSync/ConcertSyncCore/Source/ConcertSyncCore/Public/Replication/Processing/Proxy/FrequencyTracker.h b/Engine/Plugins/Developer/Concert/ConcertSync/ConcertSyncCore/Source/ConcertSyncCore/Public/Replication/Processing/Proxy/FrequencyTracker.h
--- a/Engine/Plugins/Developer/Concert/ConcertSync/ConcertSyncCore/Source/ConcertSyncCore/Public/Replication/Processing/Proxy/FrequencyTracker.h
+++ b/Engine/Plugins/Developer/Concert/ConcertSync/ConcertSyncCore/Source/ConcertSyncCore/Public/Replication/Processing/Proxy/FrequencyTracker.h
@@ -8,6 +8,8 @@
 #include "Containers/Map.h"
 #include "HAL/PlatformTime.h"
 
+#define UE_API CONCERTSYNCCORE_API
+
 struct FDateTime;
 
 namespace UE::ConcertSyncCore
@@ -20,7 +22,7 @@ namespace UE::ConcertSyncCore
 	 * You must periodically call CleanseOutdatedObjects() to avoid leaking objects that are no longer sent.
 	 * @see FFrequencyTracker_CleanByTick
 	 */
-	class CONCERTSYNCCORE_API FFrequencyTracker
+	class FFrequencyTracker
 	{
 	public:
 
@@ -41,7 +43,7 @@ namespace UE::ConcertSyncCore
 		 * 
 		 * @return Whether the object can be processed now
 		 */
-		ECanProcessObject TrackAndCheckObject(
+		UE_API ECanProcessObject TrackAndCheckObject(
 			const FConcertObjectInStreamID& ObjectPath,
 			const FConcertObjectReplicationSettings& ReplicationSettings,
 			const FPlatformSecondsTimestamp NowAsPlatformSeconds = FPlatformTime::Seconds()
@@ -54,13 +56,13 @@ namespace UE::ConcertSyncCore
 		 * @param Object The object that was processed
 		 * @param NowAsPlatformSeconds The current time - result of FPlatformTime::Seconds(). Pass this in to avoid excessive time spent on timing.
 		 */
-		void OnProcessObject(const FConcertObjectInStreamID& Object, const FPlatformSecondsTimestamp NowAsPlatformSeconds = FPlatformTime::Seconds());
+		UE_API void OnProcessObject(const FConcertObjectInStreamID& Object, const FPlatformSecondsTimestamp NowAsPlatformSeconds = FPlatformTime::Seconds());
 		
 		/**
 		 * Removes all tracked objects were last processed before CutoffTime.
 		 * @param CutoffTime A past value of FPlatformTime::Seconds()
 		 */
-		void CleanseOutdatedObjects(const FPlatformSecondsTimestamp CutoffTime);
+		UE_API void CleanseOutdatedObjects(const FPlatformSecondsTimestamp CutoffTime);
 
 	protected:
 
@@ -74,7 +76,7 @@ namespace UE::ConcertSyncCore
 	};
 
 	/** Subclass of FFrequencyTracker that you call Tick on. Every x seconds, this will call CleanseOutdatedObjects itself. */
-	class CONCERTSYNCCORE_API FFrequencyTracker_CleanByTick : public FFrequencyTracker
+	class FFrequencyTracker_CleanByTick : public FFrequencyTracker
 	{
 	public:
 
@@ -82,7 +84,7 @@ namespace UE::ConcertSyncCore
 			: CleanUpIntervalSeconds(InCleanUpInterval)
 		{}
 
-		void Tick(float DeltaTime, const FPlatformSecondsTimestamp NowAsPlatformSeconds = FPlatformTime::Seconds());
+		UE_API void Tick(float DeltaTime, const FPlatformSecondsTimestamp NowAsPlatformSeconds = FPlatformTime::Seconds());
 
 	private:
 
@@ -90,3 +92,5 @@ namespace UE::ConcertSyncCore
 		double SecondsSinceLastCleanse = 0.f;
 	};
 }
+
+#undef UE_API
diff --git a/Engine/Plugins/Developer/Concert/ConcertSync/ConcertSyncCore/Source/ConcertSyncCore/Public/Replication/Processing/ReplicationDataQueuer.h b/Engine/Plugins/Developer/Concert/ConcertSync/ConcertSyncCore/Source/ConcertSyncCore/Public/Replication/Processing/ReplicationDataQueuer.h
--- a/Engine/Plugins/Developer/Concert/ConcertSync/ConcertSyncCore/Source/ConcertSyncCore/Public/Replication/Processing/ReplicationDataQueuer.h
+++ b/Engine/Plugins/Developer/Concert/ConcertSync/ConcertSyncCore/Source/ConcertSyncCore/Public/Replication/Processing/ReplicationDataQueuer.h
@@ -5,6 +5,8 @@
 #include "IReplicationDataSource.h"
 #include "ObjectReplicationCache.h"
 
+#define UE_API CONCERTSYNCCORE_API
+
 namespace UE::ConcertSyncCore
 {
 	/**
@@ -17,7 +19,7 @@ namespace UE::ConcertSyncCore
 	 * Child classes need to implement IReplicationCacheUser::WantsToAcceptObject.
 	 * TODO: Remove object data if the object is is meant for becomes unavailable.
 	 */
-	class CONCERTSYNCCORE_API FReplicationDataQueuer
+	class FReplicationDataQueuer
 		: public IReplicationDataSource
 		, public IReplicationCacheUser
 		, public TSharedFromThis<FReplicationDataQueuer>
@@ -25,14 +27,14 @@ namespace UE::ConcertSyncCore
 	public:
 
 		//~ Begin IReplicationDataSource Interface
-		virtual void ForEachPendingObject(TFunctionRef<void(const FPendingObjectReplicationInfo&)> ProcessItemFunc) const override;
-		virtual int32 NumObjects() const override;
-		virtual bool ExtractReplicationDataForObject(const FConcertReplicatedObjectId& Object, TFunctionRef<void(const FConcertSessionSerializedPayload& Payload)> ProcessCopyable, TFunctionRef<void(FConcertSessionSerializedPayload&& Payload)> ProcessMoveable) override;
+		UE_API virtual void ForEachPendingObject(TFunctionRef<void(const FPendingObjectReplicationInfo&)> ProcessItemFunc) const override;
+		UE_API virtual int32 NumObjects() const override;
+		UE_API virtual bool ExtractReplicationDataForObject(const FConcertReplicatedObjectId& Object, TFunctionRef<void(const FConcertSessionSerializedPayload& Payload)> ProcessCopyable, TFunctionRef<void(FConcertSessionSerializedPayload&& Payload)> ProcessMoveable) override;
 		//~ End IReplicationDataSource Interface
 		
 		//~ Begin IReplicationCacheUser Interface
-		virtual void OnDataCached(const FConcertReplicatedObjectId& Object, FSequenceId SequenceId, TSharedRef<const FConcertReplication_ObjectReplicationEvent> Data) override;
-		virtual void OnCachedDataUpdated(const FConcertReplicatedObjectId& Object, FSequenceId SequenceId) override;
+		UE_API virtual void OnDataCached(const FConcertReplicatedObjectId& Object, FSequenceId SequenceId, TSharedRef<const FConcertReplication_ObjectReplicationEvent> Data) override;
+		UE_API virtual void OnCachedDataUpdated(const FConcertReplicatedObjectId& Object, FSequenceId SequenceId) override;
 		//~ End IReplicationCacheUser Interface
 
 	protected:
@@ -41,7 +43,7 @@ namespace UE::ConcertSyncCore
 		 * Called by subclass factory functions.
 		 * @param InReplicationCache The cache that the data is sent to. The caller ensures it outlives the lifetime of the widget.
 		 */
-		void BindToCache(FObjectReplicationCache& InReplicationCache);
+		UE_API void BindToCache(FObjectReplicationCache& InReplicationCache);
 		
 	private:
 
@@ -60,3 +62,5 @@ namespace UE::ConcertSyncCore
 		FObjectReplicationCache* ReplicationCache = nullptr;
 	};
 }
+
+#undef UE_API
diff --git a/Engine/Plugins/Developer/Concert/ConcertSync/ConcertSyncCore/Source/ConcertSyncCore/Public/Replication/PropertyResolutionCache.h b/Engine/Plugins/Developer/Concert/ConcertSync/ConcertSyncCore/Source/ConcertSyncCore/Public/Replication/PropertyResolutionCache.h
--- a/Engine/Plugins/Developer/Concert/ConcertSync/ConcertSyncCore/Source/ConcertSyncCore/Public/Replication/PropertyResolutionCache.h
+++ b/Engine/Plugins/Developer/Concert/ConcertSync/ConcertSyncCore/Source/ConcertSyncCore/Public/Replication/PropertyResolutionCache.h
@@ -1,10 +1,12 @@
-﻿// Copyright Epic Games, Inc. All Rights Reserved.
+// Copyright Epic Games, Inc. All Rights Reserved.
 
 #pragma once
 
 #include "Containers/Map.h"
 #include "Data/ConcertPropertySelection.h"
 
+#define UE_API CONCERTSYNCCORE_API
+
 class UStruct;
 
 namespace UE::ConcertSyncCore::PropertyChain
@@ -15,12 +17,12 @@ namespace UE::ConcertSyncCore::PropertyChain
 	 * This caches look-up results of native classes.
 	 * By default, properties in Blueprint classes are not cached because they can change at runtime and leave dangling FProperty pointers.
 	 */
-	class CONCERTSYNCCORE_API FPropertyResolutionCache
+	class FPropertyResolutionCache
 	{
 	public:
 
 		/** Resolves the property chain and caches the result; */
-		FProperty* ResolveAndCache(const UStruct& Struct, const FConcertPropertyChain& Chain);
+		UE_API FProperty* ResolveAndCache(const UStruct& Struct, const FConcertPropertyChain& Chain);
 
 		/** Removes the class from the cache.*/
 		void Invalidate(const UStruct& Struct) { CachedProperties.Remove(&Struct); }
@@ -37,3 +39,5 @@ namespace UE::ConcertSyncCore::PropertyChain
 		TMap<const UStruct*, FClassCache> CachedProperties;
 	};
 }
+
+#undef UE_API
diff --git a/Engine/Plugins/Developer/Concert/ConcertSync/ConcertSyncCore/Source/ConcertSyncCore/Public/Replication/ReplicationPropertyFilter.h b/Engine/Plugins/Developer/Concert/ConcertSync/ConcertSyncCore/Source/ConcertSyncCore/Public/Replication/ReplicationPropertyFilter.h
--- a/Engine/Plugins/Developer/Concert/ConcertSync/ConcertSyncCore/Source/ConcertSyncCore/Public/Replication/ReplicationPropertyFilter.h
+++ b/Engine/Plugins/Developer/Concert/ConcertSync/ConcertSyncCore/Source/ConcertSyncCore/Public/Replication/ReplicationPropertyFilter.h
@@ -7,6 +7,8 @@
 #include "HAL/Platform.h"
 #include "UObject/NameTypes.h"
 
+#define UE_API CONCERTSYNCCORE_API
+
 class FProperty;
 struct FArchiveSerializedPropertyChain;
 struct FConcertPropertyChain;
@@ -27,7 +29,7 @@ namespace UE::ConcertSyncCore
 	 *		Format.CreateReplicationEvent(Object, [&Filter](const FArchiveSerializedPropertyChain* Chain, const FProperty&& Property){ return Filter.IsPropertyInSelection(Chain, Property); }
 	 *	}
 	 */
-	class CONCERTSYNCCORE_API FReplicationPropertyFilter
+	class FReplicationPropertyFilter
 	{
 		const FConcertPropertySelection& PropertySelection;
 		
@@ -38,13 +40,15 @@ namespace UE::ConcertSyncCore
 		TMap<FName, TArray<const FConcertPropertyChain*>> LeafToChain;
 	public:
 
-		FReplicationPropertyFilter(const FConcertPropertySelection& PropertySelection);
+		UE_API FReplicationPropertyFilter(const FConcertPropertySelection& PropertySelection);
 		
 		/**
 		 * @param Chain Can be nullptr or non-null and empty - both imply the Property is a root property.
 		 * @param Property The property that Chain leads to
 		 * @return Whether Property is in PropertySelection.
 		 */
-		bool ShouldSerializeProperty(const FArchiveSerializedPropertyChain* Chain, const FProperty& Property) const;
+		UE_API bool ShouldSerializeProperty(const FArchiveSerializedPropertyChain* Chain, const FProperty& Property) const;
 	};
-}
\ No newline at end of file
+}
+
+#undef UE_API
