Restrict plugin dependencies from ConcertSyncCore and an internal plugin to not pull in the SQLiteCore plugin and module on console platforms.

#rnx

diff --git a/Engine/Plugins/Developer/Concert/ConcertSync/ConcertSyncCore/ConcertSyncCore.uplugin b/Engine/Plugins/Developer/Concert/ConcertSync/ConcertSyncCore/ConcertSyncCore.uplugin
--- a/Engine/Plugins/Developer/Concert/ConcertSync/ConcertSyncCore/ConcertSyncCore.uplugin
+++ b/Engine/Plugins/Developer/Concert/ConcertSync/ConcertSyncCore/ConcertSyncCore.uplugin
@@ -27,11 +27,19 @@
 	"Plugins": [
 		{
 			"Name": "ConcertMain",
-			"Enabled": true
+			"Enabled": true,
+			"TargetDenyList": [
+				"Client",
+				"Server"
+			]
 		},
 		{
 			"Name": "SQLiteCore",
-			"Enabled": true
+			"Enabled": true,
+			"TargetDenyList": [
+				"Client",
+				"Server"
+			]
 		}
 	]
 }
