Change default value for bIncludeAnnotationObjectChanges to use the updated code path for transaction diffing in Multi-user. This will allow us to deprecate the old path in the future.

diff --git a/Engine/Plugins/Developer/Concert/ConcertSync/ConcertSyncCore/Source/ConcertSyncCore/Public/ConcertSyncSettings.h b/Engine/Plugins/Developer/Concert/ConcertSync/ConcertSyncCore/Source/ConcertSyncCore/Public/ConcertSyncSettings.h
--- a/Engine/Plugins/Developer/Concert/ConcertSync/ConcertSyncCore/Source/ConcertSyncCore/Public/ConcertSyncSettings.h
+++ b/Engine/Plugins/Developer/Concert/ConcertSync/ConcertSyncCore/Source/ConcertSyncCore/Public/ConcertSyncSettings.h
@@ -65,7 +65,7 @@ public:
 	UConcertSyncConfig()
 		: bInteractiveHotReload(false)
 		, bShowPresenceInPIE(true)
-		, bIncludeAnnotationObjectChanges(false)
+		, bIncludeAnnotationObjectChanges(true)
 		, SnapshotTransactionsPerSecond(30.0f)
 	{}
 
