Fix Multi-user data store to allow for oodle compression type by splitting the compression details struct into two separate structs to encode the type and bias used for compressing concert payloads.  The previous implementation relied on a USTRUCT with bitmasks which is not supported by the serializer so oodle compression never worked properly. We now follow the pattern like in Live Link and store into different USTRUCTs.  This allows us to turn on oddle compression for payload storage.

diff --git a/Engine/Plugins/Developer/Concert/ConcertMain/Source/Concert/Private/ConcertMessageData.cpp b/Engine/Plugins/Developer/Concert/ConcertMain/Source/Concert/Private/ConcertMessageData.cpp
--- a/Engine/Plugins/Developer/Concert/ConcertMain/Source/Concert/Private/ConcertMessageData.cpp
+++ b/Engine/Plugins/Developer/Concert/ConcertMain/Source/Concert/Private/ConcertMessageData.cpp
@@ -14,6 +14,8 @@
 #include "Backends/CborStructSerializerBackend.h"
 #include "Backends/CborStructDeserializerBackend.h"
 
+#include "Math/UnitConversion.h"
+
 #include UE_INLINE_GENERATED_CPP_BY_NAME(ConcertMessageData)
 
 LLM_DEFINE_TAG(Concert_ConcertClientInfo);
@@ -173,7 +175,7 @@ namespace UE::Concert::Compression
 {
 
 static TAutoConsoleVariable<int32> CVarCompressionType(
-	TEXT("Concert.SetCompressionType"), 2,
+	TEXT("Concert.SetCompressionType"), 1,
 	TEXT("Specify the type of compression to use when serializing data. A value of 0 means compression is off. A value of 1 = Oodle. All other values = Zlib."));
 
 static TAutoConsoleVariable<int32> CVarCompressionSizeLimit(
@@ -236,21 +238,34 @@ bool TryCompressImpl(const UScriptStruct* InEventType, const void* InEventData,
 		SCOPED_CONCERT_TRACE(SerializePayload_CompressMemory);
 		if (FCompression::CompressMemory(NamedCompressionAlgo, OutCompressedData.GetData(), CompressedSize, InBytes.GetData(), InBytes.Num(), CompressFlags))
 		{
+			const FNumericUnit<uint64> InBytesDisplay = FUnitConversion::QuantizeUnitsToBestFit(InBytes.Num(), EUnit::Bytes);
+			const FNumericUnit<uint64> OutBytesDisplay = FUnitConversion::QuantizeUnitsToBestFit(CompressedSize, EUnit::Bytes);
+			UE_LOG(LogConcert, VeryVerbose, TEXT("Compressed payload for %s from to %lld %s to %lld %s. Compression Type: %s"), *InEventType->GetName(), InBytesDisplay.Value,
+				   FUnitConversion::GetUnitDisplayString(InBytesDisplay.Units), OutBytesDisplay.Value, FUnitConversion::GetUnitDisplayString(OutBytesDisplay.Units),
+				   *NamedCompressionAlgo.ToString());   
+			
 			OutCompressedData.SetNum(CompressedSize, EAllowShrinking::No);
 			InOutPayload.PayloadBytes.Bytes = MoveTemp(OutCompressedData);
-			InOutPayload.PayloadCompressionDetails = UE::Concert::Compression::GetCompressionFromNamedType(NamedCompressionAlgo, CompressFlags);
+			InOutPayload.PayloadCompressionMethod = UE::Concert::Compression::GetCompressionFromNamedType(NamedCompressionAlgo);
+			InOutPayload.PayloadCompressionBias = UE::Concert::Compression::GetConcertCompressionBias(CompressFlags);
 		}
 		else
 		{
 			UE_LOG(LogConcert, Warning, TEXT("Unable to compress data for %s!"), *InEventType->GetName());
-			InOutPayload.PayloadCompressionDetails = EConcertCompressionDetails::Uncompressed;
+			InOutPayload.PayloadCompressionMethod = EConcertCompressionMethod::Uncompressed;
+			InOutPayload.PayloadCompressionBias = EConcertCompressionBias::None;
 		}
 	}
 	else
 	{
-		InOutPayload.PayloadCompressionDetails = EConcertCompressionDetails::Uncompressed;
-	}
+		const FNumericUnit<uint64> InBytesDisplay = FUnitConversion::QuantizeUnitsToBestFit(InOutPayload.PayloadBytes.Bytes.Num(), EUnit::Bytes);
+		UE_LOG(LogConcert, VeryVerbose, TEXT("Not compressing payload for %s. Size is %lld %s"),
+			   *InEventType->GetName(), InBytesDisplay.Value, FUnitConversion::GetUnitDisplayString(InBytesDisplay.Units));
 
+		InOutPayload.PayloadCompressionMethod = EConcertCompressionMethod::Uncompressed;
+		InOutPayload.PayloadCompressionBias = EConcertCompressionBias::None;
+	}
+			
 	// Since we can support uncompressed or compressed data this is always successful.
 	return true;
 }
@@ -259,14 +274,14 @@ using OptionalDecompressBytes = TOptional<TArray<uint8>>;
 
 OptionalDecompressBytes DecompressImpl(const FConcertSessionSerializedPayload& InPayload)
 {
-	EConcertCompressionDetails CompressAlgo = InPayload.PayloadCompressionDetails;
+	EConcertCompressionMethod CompressAlgo = InPayload.PayloadCompressionMethod;
 	if (UE::Concert::Compression::DataIsCompressed(CompressAlgo))
 	{
 		const TArray<uint8> &InBytes = InPayload.PayloadBytes.Bytes;
 		TArray<uint8> UncompressedData;
 		UncompressedData.SetNumUninitialized(InPayload.PayloadSize);
 		SCOPED_CONCERT_TRACE(DeserializePayload_UncompressMemory);
-		ECompressionFlags CompressFlags = UE::Concert::Compression::GetCoreCompressionFlags(CompressAlgo);
+		ECompressionFlags CompressFlags = UE::Concert::Compression::GetCoreCompressionFlags(InPayload.PayloadCompressionBias);
 		FName CompressType = UE::Concert::Compression::GetCompressionAlgorithm(CompressAlgo);
 		if (FCompression::UncompressMemory(CompressType, UncompressedData.GetData(), UncompressedData.Num(), InBytes.GetData(), InBytes.Num(), CompressFlags))
 		{
@@ -319,7 +334,7 @@ bool DeserializeImpl(const UScriptStruct* InTargetEventType, void* InOutTargetEv
 bool DeserializeAndDecompress(const UScriptStruct* InTargetEventType, void* InOutTargetEventData, const FConcertSessionSerializedPayload& InPayload)
 {
 	OptionalDecompressBytes DecompressedBytes = DecompressImpl(InPayload);
-	EConcertCompressionDetails CompressAlgo = InPayload.PayloadCompressionDetails;
+	EConcertCompressionMethod CompressAlgo = InPayload.PayloadCompressionMethod;
 
 	if ( UE::Concert::Compression::DataIsCompressed(CompressAlgo) && !DecompressedBytes.IsSet() )
 	{
diff --git a/Engine/Plugins/Developer/Concert/ConcertMain/Source/Concert/Public/ConcertMessageData.h b/Engine/Plugins/Developer/Concert/ConcertMain/Source/Concert/Public/ConcertMessageData.h
--- a/Engine/Plugins/Developer/Concert/ConcertMain/Source/Concert/Public/ConcertMessageData.h
+++ b/Engine/Plugins/Developer/Concert/ConcertMain/Source/Concert/Public/ConcertMessageData.h
@@ -46,7 +46,7 @@ enum class EConcertPayloadCompressionType : uint8
 };
 
 UENUM()
-enum class EConcertCompressionDetails : uint8
+enum class UE_DEPRECATED(5.8, "Use EConcertCompressionMethod instead.") EConcertCompressionDetails : uint8
 {
 	/** The package data (the body) is written uncompressed. */
 	Uncompressed = 0,
@@ -63,7 +63,36 @@ enum class EConcertCompressionDetails : uint8
 	/** Compress for smaller sizes. */
 	CompressForSize = 1 << 3
 };
-ENUM_CLASS_FLAGS(EConcertCompressionDetails)
+
+/** Compression method to use. Note that these are assigned a bit shifted value to match the previous implementation in EConcertCompressionDetails. */
+UENUM()
+enum class EConcertCompressionMethod : uint8
+{
+	/** Data is uncompressed. */
+	Uncompressed = 0,
+
+	/**  ZLib compression is used for the serialized data. */
+	Zlib = 1 << 0,
+	
+	/** Oodle compression is used for the serialized data. */
+	Oodle = 1 << 1
+};
+
+/** The bias (if any) to use when compressing the data. Note that these are assigned a bit shifted value to match the previous implementation in EConcertCompressionDetails. */
+UENUM()
+enum class EConcertCompressionBias : uint8
+{
+	/** Compress without a bias. */
+	None,
+
+	/** Compress with a bias for speed. */
+	Speed = 1 << 2,
+
+	/** Compress with a bias for size. */
+	Size = 1 << 3
+};
+
+
 
 namespace UE::Concert::Compression
 {
@@ -93,21 +122,21 @@ inline bool ShouldCompress(SizeType DataSize)
 }
 
 /** Returns truen if the compression details is set to compressed format.  */
-inline bool DataIsCompressed(EConcertCompressionDetails InFormat)
+inline bool DataIsCompressed(EConcertCompressionMethod InFormat)
 {
-	return EnumHasAllFlags(InFormat, EConcertCompressionDetails::Compressed);
+	return InFormat == EConcertCompressionMethod::Zlib || InFormat == EConcertCompressionMethod::Oodle;
 }
 
 /** Returns true if the compression details is set to uncompressed format. */
-inline bool DataIsUncompressed(EConcertCompressionDetails InFormat)
+inline bool DataIsUncompressed(EConcertCompressionMethod InFormat)
 {
-	return InFormat == EConcertCompressionDetails::Uncompressed;
+	return InFormat == EConcertCompressionMethod::Uncompressed;
 }
 
 /** Returns true if we are suppose to use Oodle for compression. */
-inline bool DataIsCompressedWithOodle(EConcertCompressionDetails InFormat)
+inline bool DataIsCompressedWithOodle(EConcertCompressionMethod InFormat)
 {
-	return EnumHasAllFlags(InFormat, EConcertCompressionDetails::Compressed | EConcertCompressionDetails::CompressWithOodle);
+	return InFormat == EConcertCompressionMethod::Oodle;
 }
 
 /** Get the named compression algorithm to invoke with serializer and memory compressors. */
@@ -137,73 +166,59 @@ inline ECompressionFlags GetCompressionFlags()
 }
 
 /** Get the named compression algorithm based on the provided details.  Current this is only Zlib or Oodle. */
-inline FName GetCompressionAlgorithm(EConcertCompressionDetails InDetails)
+inline FName GetCompressionAlgorithm(EConcertCompressionMethod InDetails)
 {
 	return DataIsCompressedWithOodle(InDetails) ? NAME_Oodle : NAME_Zlib;
 }
 
 /** Get the compression enum value given the concert compression settings.  This the value passed into the serializer and memory compressor. */
-inline ECompressionFlags GetCoreCompressionFlags(EConcertCompressionDetails InFormat)
+inline ECompressionFlags GetCoreCompressionFlags(EConcertCompressionBias InBias)
 {
-	// Make sure we do not have both flags set.
-	check(!EnumHasAllFlags(InFormat, EConcertCompressionDetails::CompressForSize | EConcertCompressionDetails::CompressForSpeed));
-	if (DataIsUncompressed(InFormat))
-	{
-		return COMPRESS_NoFlags;
-	}
-	if (EnumHasAnyFlags(InFormat, EConcertCompressionDetails::CompressForSize))
+	if (InBias == EConcertCompressionBias::Size)
 	{
 		return COMPRESS_BiasSize;
 	}
-	if (EnumHasAnyFlags(InFormat, EConcertCompressionDetails::CompressForSpeed))
+	if (InBias == EConcertCompressionBias::Speed)
 	{
 		return COMPRESS_BiasSpeed;
 	}
 	return COMPRESS_NoFlags;
 }
 
-/** For a given named compression algorithm and compression flags convert it into a EConcertCompressionDetails */
-inline EConcertCompressionDetails GetCompressionFromNamedType(FName NamedMethod, ECompressionFlags Flags)
+/** Get the concert compression bias value given the compressor compression flags. */
+inline EConcertCompressionBias GetConcertCompressionBias(ECompressionFlags InFlags)
 {
-	check(NamedMethod == NAME_Oodle || NamedMethod == NAME_Zlib);
-	check(Flags == COMPRESS_NoFlags || Flags == COMPRESS_BiasMemory || Flags == COMPRESS_BiasSpeed);
-
-	EConcertCompressionDetails Compressor = EConcertCompressionDetails::Compressed;
-
-	if (NamedMethod == NAME_Oodle)
+	if (InFlags == COMPRESS_BiasSize)
 	{
-		Compressor = Compressor | EConcertCompressionDetails::CompressWithOodle;
+		return EConcertCompressionBias::Size;
 	}
-
-	if (Flags == COMPRESS_BiasMemory)
-	{
-		Compressor = Compressor | EConcertCompressionDetails::CompressForSize;
-	}
-	else if (Flags == COMPRESS_BiasSpeed)
+	
+	if (InFlags == COMPRESS_BiasSpeed)
 	{
-		Compressor = Compressor | EConcertCompressionDetails::CompressForSpeed;
+		return EConcertCompressionBias::Speed;
 	}
-	return Compressor;
+	
+	return EConcertCompressionBias::None;
 }
 
-/** Add in the compression flags to the concert compression details enum. */
-inline EConcertCompressionDetails GetCompressionFlags(EConcertCompressionDetails InFormat)
+/** For a given named compression algorithm and compression flags convert it into a EConcertCompressionMethod */
+inline EConcertCompressionMethod GetCompressionFromNamedType(FName NamedMethod)
 {
-	const int32 CompressionFlags = GetConsoleVariableCompressionFlags();
-	if (CompressionFlags == 0)
-	{
-		return InFormat;
-	}
-	if (CompressionFlags == 1)
+	check(NamedMethod == NAME_Oodle || NamedMethod == NAME_Zlib);
+
+	EConcertCompressionMethod Compressor = EConcertCompressionMethod::Zlib;
+
+	if (NamedMethod == NAME_Oodle)
 	{
-		return InFormat | EConcertCompressionDetails::CompressForSize;
+		Compressor = EConcertCompressionMethod::Oodle;
 	}
-	return InFormat | EConcertCompressionDetails::CompressForSpeed;
+
+	return Compressor;
 }
 
 /** Based on the data size return the current compression details to store. */
 template <typename SizeType>
-inline EConcertCompressionDetails GetCompressionDetails(SizeType DataSize)
+inline EConcertCompressionMethod GetCompressionMethod(SizeType DataSize)
 {
 	static_assert(std::is_integral<SizeType>::value);
 	const int32 CompressionType = GetConsoleVariableCompressionType();
@@ -211,11 +226,11 @@ inline EConcertCompressionDetails GetCompressionDetails(SizeType DataSize)
 	{
 		if (CompressionType == 1)
 		{
-			return GetCompressionFlags(EConcertCompressionDetails::Compressed | EConcertCompressionDetails::CompressWithOodle);
+			return EConcertCompressionMethod::Oodle;
 		}
-		return GetCompressionFlags(EConcertCompressionDetails::Compressed);
+		return EConcertCompressionMethod::Zlib;
 	}
-	return EConcertCompressionDetails::Uncompressed;
+	return EConcertCompressionMethod::Uncompressed;
 }
 
 }
@@ -534,7 +549,7 @@ struct FConcertSessionSerializedPayload
 	/** Indicates if the stored payload has been compressed with any compression algorithm. */
 	bool PayloadIsCompressed() const
 	{
-		return UE::Concert::Compression::DataIsCompressed(PayloadCompressionDetails);
+		return UE::Concert::Compression::DataIsCompressed(PayloadCompressionMethod);
 	}
 
 	/** The typename of the user-defined payload. */
@@ -547,8 +562,12 @@ struct FConcertSessionSerializedPayload
 
 	/** Indicates how the serialized payload has been compressed. */
 	UPROPERTY(VisibleAnywhere, Category = "Payload")
-	EConcertCompressionDetails PayloadCompressionDetails = EConcertCompressionDetails::Uncompressed;
+	EConcertCompressionMethod PayloadCompressionMethod = EConcertCompressionMethod::Uncompressed;
 
+	/** Indicates how the serialized payload has been compressed. */
+	UPROPERTY(VisibleAnywhere, Category = "Payload")
+	EConcertCompressionBias PayloadCompressionBias = EConcertCompressionBias::None;
+    
 	/** The uncompressed size of the user-defined payload data. */
 	UPROPERTY(VisibleAnywhere, Category="Payload")
 	int32 PayloadSize = 0;
diff --git a/Engine/Plugins/Developer/Concert/ConcertSync/ConcertSyncClient/Source/ConcertSyncClient/Private/ConcertClientDataStore.cpp b/Engine/Plugins/Developer/Concert/ConcertSync/ConcertSyncClient/Source/ConcertSyncClient/Private/ConcertClientDataStore.cpp
--- a/Engine/Plugins/Developer/Concert/ConcertSync/ConcertSyncClient/Source/ConcertSyncClient/Private/ConcertClientDataStore.cpp
+++ b/Engine/Plugins/Developer/Concert/ConcertSync/ConcertSyncClient/Source/ConcertSyncClient/Private/ConcertClientDataStore.cpp
@@ -99,7 +99,8 @@ TFuture<FConcertDataStoreResult> FConcertClientDataStore::InternalCompareExchang
 			// Send the 'version' rather than 'expected' to save bandwidth.
 			CompareExchangeRequest.ExpectedVersion = Result.Value->Version;
 			CompareExchangeRequest.Expected.PayloadTypeName = TEXT("");
-			CompareExchangeRequest.Expected.PayloadCompressionDetails = EConcertCompressionDetails::Uncompressed;
+			CompareExchangeRequest.Expected.PayloadCompressionMethod = EConcertCompressionMethod::Uncompressed;
+			CompareExchangeRequest.Expected.PayloadCompressionBias = EConcertCompressionBias::None;
 			CompareExchangeRequest.Expected.PayloadSize = 0;
 			CompareExchangeRequest.Expected.PayloadBytes.Bytes.Reset();
 		}
@@ -185,7 +186,8 @@ FConcertDataStoreResult FConcertClientDataStore::HandleResponse(const FName& Sen
 		// Ensure the server doesn't send data when the client initiated the operation (The client knows the value it sent).
 		check(Response.Value.TypeName.IsNone());
 		check(Response.Value.SerializedValue.PayloadTypeName.IsNone());
-		check(Response.Value.SerializedValue.PayloadCompressionDetails == EConcertCompressionDetails::Uncompressed);
+		check(Response.Value.SerializedValue.PayloadCompressionMethod == EConcertCompressionMethod::Uncompressed);
+		check(Response.Value.SerializedValue.PayloadCompressionBias == EConcertCompressionBias::None);
 		check(Response.Value.SerializedValue.PayloadBytes.Bytes.Num() == 0);
 
 		// Ensure the server sent back a valid version in case it has exchanged the value. (Successfully added value is always version 1)
diff --git a/Engine/Plugins/Developer/Concert/ConcertSync/ConcertSyncCore/Source/ConcertSyncCore/Private/ConcertSyncSessionDatabase.cpp b/Engine/Plugins/Developer/Concert/ConcertSync/ConcertSyncCore/Source/ConcertSyncCore/Private/ConcertSyncSessionDatabase.cpp
--- a/Engine/Plugins/Developer/Concert/ConcertSync/ConcertSyncCore/Source/ConcertSyncCore/Private/ConcertSyncSessionDatabase.cpp
+++ b/Engine/Plugins/Developer/Concert/ConcertSync/ConcertSyncCore/Source/ConcertSyncCore/Private/ConcertSyncSessionDatabase.cpp
@@ -25,6 +25,90 @@ namespace PackageDataUtil
 	};
 }
 
+namespace TransactionDataUtil
+{
+
+const int32 MinFilesToCache = 10;
+const uint64 MaxFileSizeBytesToCache = 50 * 1024 * 1024;
+const int64 BucketSize = 500;
+const uint32 DataVersion = 2;
+const FGuid DataEntryFooter = FGuid(0xE473C070, 0x65DA42BF, 0xA0607C78, 0xE0DC47CF);
+
+}; // namespace TransactionDataUtil
+
+struct FCompressionMethodAndBias
+{
+	FCompressionMethodAndBias() = default;
+	FCompressionMethodAndBias(EConcertCompressionMethod InMethod, EConcertCompressionBias InBias)
+		:  Method(InMethod)
+		  ,Bias(InBias)
+	{
+	}
+
+	
+	EConcertCompressionMethod Method = EConcertCompressionMethod::Uncompressed;
+	EConcertCompressionBias Bias = EConcertCompressionBias::None;
+
+
+	/** Non-shifted conversion into a int32 value. This version is used when storing into the database record as all values are int32 */
+	uint32 ToInt32() const
+	{
+		const uint32 CompressionMethodAndBias = static_cast<int32>(Method) | static_cast<int32>(Bias);
+		return CompressionMethodAndBias;
+	}
+
+	/** Non-shifted conversion into a enum types. This version is used when reading a int32 table entry that corresponds to the compression flags. */
+	void FromInt32(int32 InCompressionAndBias)
+	{
+		Method = static_cast<EConcertCompressionMethod>(InCompressionAndBias & (static_cast<int32>(EConcertCompressionMethod::Zlib) | static_cast<int32>(EConcertCompressionMethod::Oodle)));
+		Bias = static_cast<EConcertCompressionBias>(InCompressionAndBias & (static_cast<int32>(EConcertCompressionBias::Speed) | static_cast<int32>(EConcertCompressionBias::Size)));
+		check(IsValid());
+	}
+
+	/** Non-shifted conversion into a uint32 value. This version is used when serializing into a transaction record. */
+	uint32 ToUint32() const
+	{
+		const uint32 CompressionMethodAndBias = static_cast<uint32>(Method) | static_cast<uint32>(Bias);
+		return CompressionMethodAndBias;
+	}
+
+	/** Non-shifted conversion into a enum types. This version is used when deserializing from a transaction record. */
+	void FromUint32(uint32 InCompressionAndBias)
+	{
+		Method = static_cast<EConcertCompressionMethod>(InCompressionAndBias & (static_cast<uint32>(EConcertCompressionMethod::Zlib) | static_cast<uint32>(EConcertCompressionMethod::Oodle)));
+		Bias = static_cast<EConcertCompressionBias>(InCompressionAndBias & (static_cast<uint32>(EConcertCompressionBias::Speed) | static_cast<uint32>(EConcertCompressionBias::Size)));
+		check(IsValid());
+	}
+
+	bool IsValid()
+	{
+		const bool bMethodValid = (Method == EConcertCompressionMethod::Uncompressed
+								|| Method == EConcertCompressionMethod::Oodle
+								|| Method == EConcertCompressionMethod::Zlib);
+		const bool bBiasValid = ( Bias == EConcertCompressionBias::None
+								|| Bias == EConcertCompressionBias::Speed
+								|| Bias == EConcertCompressionBias::Size);
+		return bMethodValid && bBiasValid;
+	}
+
+	/** Shifted conversion which is used when writing the package blob. We write this in the 24 MSB of uint32 */
+	uint32 ToUint32_Shifted() const
+	{
+		const uint32 CompressionMethodAndBias = static_cast<uint8>(Method) | static_cast<uint8>(Bias);
+		return CompressionMethodAndBias << 8;
+	}
+
+	/** Convert a method and bias value that was stored in the 24 MSB back into enum values. */
+	void FromUint32_Shifted(uint32 InSerializedDataVersion)
+	{
+		const uint32 CompressionAndBias = (InSerializedDataVersion >> 8);
+	
+		Method = static_cast<EConcertCompressionMethod>(CompressionAndBias & (static_cast<uint32>(EConcertCompressionMethod::Zlib) | static_cast<uint32>(EConcertCompressionMethod::Oodle)));
+		Bias = static_cast<EConcertCompressionBias>(CompressionAndBias & (static_cast<uint32>(EConcertCompressionBias::Speed) | static_cast<uint32>(EConcertCompressionBias::Size)));
+		check(IsValid());
+	}
+};
+
 struct FConcertPackageAsyncDataStream
 {
 	FConcertPackageAsyncDataStream(const FConcertPackageDataStream& InStream, const FName& InPackageName)
@@ -99,12 +183,6 @@ public:
 namespace TransactionDataUtil
 {
 
-const int32 MinFilesToCache = 10;
-const uint64 MaxFileSizeBytesToCache = 50 * 1024 * 1024;
-const int64 BucketSize = 500;
-const uint32 DataVersion = 2;
-const FGuid DataEntryFooter = FGuid(0xE473C070, 0x65DA42BF, 0xA0607C78, 0xE0DC47CF);
-
 FString GetDataPath(const FString& InSessionPath)
 {
 	return InSessionPath / TEXT("Transactions");
@@ -152,14 +230,16 @@ bool WriteTransaction(const FStructOnScope& InTransaction, TArray<uint8>& OutSer
 	Ar.SerializeIntPacked(UncompressedTransactionSize);
 	if (UncompressedTransactionSize > 0)
 	{
-		EConcertCompressionDetails StoredFormat = UE::Concert::Compression::GetCompressionDetails(UncompressedTransactionSize);
-		uint32 StoredFormatAsUint = static_cast<uint32>(StoredFormat);
-		Ar.SerializeIntPacked(StoredFormatAsUint);
-		if (UE::Concert::Compression::DataIsCompressed(StoredFormat))
+		FCompressionMethodAndBias MethodAndBias(UE::Concert::Compression::GetCompressionMethod(UncompressedTransactionSize),
+												UE::Concert::Compression::GetConcertCompressionBias(UE::Concert::Compression::GetCompressionFlags()));  
+
+		uint32 StoredFormat = MethodAndBias.ToUint32();
+		Ar.SerializeIntPacked(StoredFormat);
+		if (UE::Concert::Compression::DataIsCompressed(MethodAndBias.Method))
 		{
 			Ar.SerializeCompressed(UncompressedTransaction.GetData(), UncompressedTransactionSize,
-								   UE::Concert::Compression::GetCompressionAlgorithm(StoredFormat),
-								   UE::Concert::Compression::GetCoreCompressionFlags(StoredFormat));
+								   UE::Concert::Compression::GetCompressionAlgorithm(MethodAndBias.Method),
+								   UE::Concert::Compression::GetCoreCompressionFlags(MethodAndBias.Bias));
 		}
 		else
 		{
@@ -243,19 +323,20 @@ bool ReadTransaction(const TArray<uint8>& InSerializedTransactionData, FStructOn
 	if (UncompressedTransactionSize > 0)
 	{
 		// Version 1 data is stored always compressed in Zlib format.
-		EConcertCompressionDetails Format = EConcertCompressionDetails::Compressed;
+		FCompressionMethodAndBias Format(EConcertCompressionMethod::Zlib, EConcertCompressionBias::None);
 		if (SerializedDataVersion > 1)
 		{
 			uint32 SerializedDataFormat = 0;
 			Ar.SerializeIntPacked(SerializedDataFormat);
-			uint8 DataFormatAsUint8 = static_cast<uint8>(SerializedDataFormat & 0xFF);
-			Format = static_cast<EConcertCompressionDetails>(DataFormatAsUint8);
+			
+			Format.FromUint32(SerializedDataFormat);
 		}
-		if (UE::Concert::Compression::DataIsCompressed(Format))
+
+		if (UE::Concert::Compression::DataIsCompressed(Format.Method))
 		{
 			Ar.SerializeCompressed(UncompressedTransaction.GetData(), UncompressedTransactionSize,
-								   UE::Concert::Compression::GetCompressionAlgorithm(Format),
-								   UE::Concert::Compression::GetCoreCompressionFlags(Format));
+								   UE::Concert::Compression::GetCompressionAlgorithm(Format.Method),
+								   UE::Concert::Compression::GetCoreCompressionFlags(Format.Bias));
 		}
 		else
 		{
@@ -286,7 +367,6 @@ const FGuid EntryFooter = FGuid(0x2EFC8CDD, 0x748E46C0, 0xA5485769, 0x13A3C354);
 constexpr const uint64 MaxPackageBlobSizeForCaching = FMath::Max(
 	PackageDataUtil::MaxFileSizeBytesToCache / PackageDataUtil::MinFilesToCache, 16ull * 1024 * 1024);
 
-
 /**
  * Returns true if a package blob should be cached in memory, according to its size.
  * @param PackageBlobSize The blob size (containing package data (possibly compressed) and some meta data).
@@ -298,18 +378,18 @@ bool ShouldCachePackageBlob(uint64 PackageBlobSize)
 }
 
 /** Write the version/format information in the first 4 bytes of the blob, enabling the reader to parse the blob. */
-void WritePackageBlobVersionInfo(FArchive& Ar, EConcertCompressionDetails PackageDataFormat)
+void WritePackageBlobVersionInfo(FArchive& Ar, FCompressionMethodAndBias InCompressionMethodAndBias)
 {
 	check(Ar.IsSaving() && Ar.Tell() == 0); // Expects to write the version/format info at the beginning.
 
-	// Keep the major DataVersion in the 8 LSB, encode package data format in the 24 MSB.
-	constexpr uint32 BitsPerByte = 8;
-	uint32 SerializedDataVersion = static_cast<uint32>(DataVersion) | (static_cast<uint32>(PackageDataFormat) << (sizeof(DataVersion) * BitsPerByte));
+	// Keep the major DataVersion in the 8 LSB, encode package data format in the 24 MSB.				
+	uint32 SerializedDataVersion = static_cast<uint32>(PackageDataUtil::DataVersion) | InCompressionMethodAndBias.ToUint32_Shifted();
 	Ar.SerializeIntPacked(SerializedDataVersion);
 }
 
+
 /** Reads the first 4 bytes of the archive and detect how the package data was written in the blob. */
-EConcertCompressionDetails ParsePackageDataFormat(FArchive& Ar)
+FCompressionMethodAndBias ParsePackageDataFormat(FArchive& Ar)
 {
 	check(Ar.IsLoading() && Ar.Tell() == 0); // Expected to read the version/format info at the beginning.
 
@@ -319,12 +399,13 @@ EConcertCompressionDetails ParsePackageDataFormat(FArchive& Ar)
 	if (ExtractedDataVersion == 1)
 	{
 		// In Version 1, the package was always written compressed into a single block.
-		return EConcertCompressionDetails::Compressed;
+		return FCompressionMethodAndBias(EConcertCompressionMethod::Zlib, EConcertCompressionBias::None);
 	}
 
 	// In Version 2, the package data format was encoded after the 8 LSB used for 'DataVersion'
-	constexpr uint32 BitsPerByte = 8;
-	return static_cast<EConcertCompressionDetails>(SerializedDataVersion >> (sizeof(DataVersion) * BitsPerByte));
+	FCompressionMethodAndBias MethodAndBias;
+	MethodAndBias.FromUint32_Shifted(SerializedDataVersion & 0xFFFFFF00);
+	return MethodAndBias;
 }
 
 FString GetDataPath(const FString& InSessionPath)
@@ -343,9 +424,10 @@ FString GetDataFilename(const FName InPackageName, const int64 InRevision)
 }
 
 /** Internally used by WritePackageBlob() - write the serialized package header. */
-void WritePackageHeader(FArchive& Ar, EConcertCompressionDetails PackageDataFormat)
+void WritePackageHeader(FArchive& Ar, EConcertCompressionMethod CompressionMethod, EConcertCompressionBias CompressionBias)
 {
 	// First 4-bytes contains information about the format.
+	FCompressionMethodAndBias PackageDataFormat(CompressionMethod, CompressionBias);
 	WritePackageBlobVersionInfo(Ar, PackageDataFormat);
 
 	// PackageDataOffset is kept to allow 4.25 to read 4.24 packages. The variable-length data stored between version info and package data was never used, so it was stripped in 4.25.
@@ -358,10 +440,10 @@ void WritePackageHeader(FArchive& Ar, EConcertCompressionDetails PackageDataForm
 }
 
 /** Internally used by ExtractPackageData() - read serialized package header. */
-TTuple<EConcertCompressionDetails, int64> ReadPackageHeader(FArchive& Ar)
+TTuple<FCompressionMethodAndBias, int64> ReadPackageHeader(FArchive& Ar)
 {
 	// Consume the first bytes to read the package version and package data format.
-	EConcertCompressionDetails PackageDataFormat = ParsePackageDataFormat(Ar);
+	FCompressionMethodAndBias PackageDataFormat = ParsePackageDataFormat(Ar);
 
 	// Deserialize the next 8 bytes containing the offset to the package data (for binary compatibility with 4.24)
 	int64 PackageDataOffset = 0;
@@ -401,11 +483,11 @@ int64 EstimatePackageBlobMetaDataSize()
 bool WillUseMemoryWriter(int64 DataSize)
 {
 	using namespace UE::Concert::Compression;
-	return DataSize > 0 && (DataIsCompressed(GetCompressionDetails(DataSize)) ||
+	return DataSize > 0 && (DataIsCompressed(GetCompressionMethod(DataSize)) ||
 							ShouldCachePackageBlob(DataSize + PackageDataUtil::EstimatePackageBlobMetaDataSize()));
 }
 
-void WritePackageCompressed(FConcertPackageDataStream& PackageDataStream, FArchive& DstAr, EConcertCompressionDetails DataFormat)
+void WritePackageCompressed(FConcertPackageDataStream& PackageDataStream, FArchive& DstAr, FCompressionMethodAndBias DataFormat)
 {
 	check(PackageDataStream.DataSize > 0);
 
@@ -417,8 +499,8 @@ void WritePackageCompressed(FConcertPackageDataStream& PackageDataStream, FArchi
 	if (PackageDataStream.DataBlob) // Optimization to avoid transfering from the archive to a temporary buffer.
 	{
 		DstAr.SerializeCompressed(const_cast<uint8*>(PackageDataStream.DataBlob->GetData()), UncompressedPackageSize,
-								  UE::Concert::Compression::GetCompressionAlgorithm(DataFormat),
-								  UE::Concert::Compression::GetCoreCompressionFlags(DataFormat));
+								  UE::Concert::Compression::GetCompressionAlgorithm(DataFormat.Method),
+								  UE::Concert::Compression::GetCoreCompressionFlags(DataFormat.Bias));
 	}
 	else
 	{
@@ -426,8 +508,8 @@ void WritePackageCompressed(FConcertPackageDataStream& PackageDataStream, FArchi
 		UncompressedPackageData.AddUninitialized(UncompressedPackageSize);
 		PackageDataStream.DataAr->Serialize(UncompressedPackageData.GetData(), UncompressedPackageSize); // Read.
 		DstAr.SerializeCompressed(const_cast<uint8*>(UncompressedPackageData.GetData()), UncompressedPackageSize,
-								  UE::Concert::Compression::GetCompressionAlgorithm(DataFormat),
-								  UE::Concert::Compression::GetCoreCompressionFlags(DataFormat));
+								  UE::Concert::Compression::GetCompressionAlgorithm(DataFormat.Method),
+								  UE::Concert::Compression::GetCoreCompressionFlags(DataFormat.Bias));
 	}
 }
 
@@ -453,7 +535,9 @@ WritePackageResult WritePackage(FConcertPackageDataStream& PackageDataStream, FA
 	WritePackageResult OutResult;
 	FMemorySyncDbWriter MemWriter(OutResult.CacheData);
 
-	EConcertCompressionDetails PackageDataFormat = UE::Concert::Compression::GetCompressionDetails(PackageDataStream.DataSize);
+	EConcertCompressionMethod CompressionMethod = UE::Concert::Compression::GetCompressionMethod(PackageDataStream.DataSize);
+	EConcertCompressionBias CompressionBias = UE::Concert::Compression::GetConcertCompressionBias(UE::Concert::Compression::GetCompressionFlags());
+
 	if (WillUseMemoryWriter(PackageDataStream.DataSize))
 	{
 		DstAr = &MemWriter;
@@ -467,10 +551,10 @@ WritePackageResult WritePackage(FConcertPackageDataStream& PackageDataStream, FA
 	 // If the package data archive is null, the size must be zero.
 	check(PackageDataStream.DataAr != nullptr || (PackageDataStream.DataAr == nullptr && PackageDataStream.DataSize == 0 && PackageDataStream.DataBlob == nullptr));
 
-	WritePackageHeader(*DstAr, PackageDataFormat);
-	if (UE::Concert::Compression::DataIsCompressed(PackageDataFormat))
+	WritePackageHeader(*DstAr, CompressionMethod, CompressionBias);
+	if (UE::Concert::Compression::DataIsCompressed(CompressionMethod))
 	{
-		WritePackageCompressed(PackageDataStream, *DstAr, PackageDataFormat);
+		WritePackageCompressed(PackageDataStream, *DstAr, FCompressionMethodAndBias(CompressionMethod, CompressionBias));
 		OutResult.bShouldCache = ShouldCachePackageBlob(OutResult.CacheData.Num());
 	}
 	else
@@ -493,15 +577,15 @@ bool ExtractPackageData(FArchive& PackageBlobAr, const TFunctionRef<void(FConcer
 	}
 
 	// Read the header.
-	TTuple<EConcertCompressionDetails, int64> Header = ReadPackageHeader(PackageBlobAr);
-	EConcertCompressionDetails PackageFormat = Header.Get<0>();
+	TTuple<FCompressionMethodAndBias, int64> Header = ReadPackageHeader(PackageBlobAr);
+	FCompressionMethodAndBias PackageFormat = Header.Get<0>();
 	int64 PackageDataOffset = Header.Get<1>();
 
 	// Put the stream in position.
 	PackageBlobAr.Seek(PackageDataOffset);
 
 	// The package data is compressed (this hints the package is small enough to be loaded in memory)
-	if (UE::Concert::Compression::DataIsCompressed(PackageFormat))
+	if (UE::Concert::Compression::DataIsCompressed(PackageFormat.Method))
 	{
 		// Read the uncompressed package data size.
 		uint32 UncompressedPackageSize = 0;
@@ -514,8 +598,8 @@ bool ExtractPackageData(FArchive& PackageBlobAr, const TFunctionRef<void(FConcer
 		{
 			DecompressedData.AddUninitialized(UncompressedPackageSize);
 			PackageBlobAr.SerializeCompressed(DecompressedData.GetData(), UncompressedPackageSize,
-											  UE::Concert::Compression::GetCompressionAlgorithm(PackageFormat),
-											  UE::Concert::Compression::GetCoreCompressionFlags(PackageFormat));
+											  UE::Concert::Compression::GetCompressionAlgorithm(PackageFormat.Method),
+											  UE::Concert::Compression::GetCoreCompressionFlags(PackageFormat.Bias));
 		}
 
 		// Let the caller stream the package data from a memory archive.
@@ -525,7 +609,7 @@ bool ExtractPackageData(FArchive& PackageBlobAr, const TFunctionRef<void(FConcer
 	}
 	else
 	{
-		check(UE::Concert::Compression::DataIsUncompressed(PackageFormat));
+		check(UE::Concert::Compression::DataIsUncompressed(PackageFormat.Method));
 
 		// Read the package size.
 		int64 PackageDataSize = 0;
@@ -555,13 +639,18 @@ bool ExtractPackageData(const FString& InPackageBlobPathname, const TFunctionRef
 template <typename SerializedDataT>
 void SetSerializedPayloadFlags(SerializedDataT& Data, int32 Flags)
 {
-	Data.PayloadCompressionDetails = static_cast<EConcertCompressionDetails>(Flags);
+	FCompressionMethodAndBias Converted;
+	Converted.FromInt32(Flags);
+	
+	Data.PayloadCompressionMethod = Converted.Method;
+	Data.PayloadCompressionBias = Converted.Bias;
 }
 
 template <typename SerializedDataT>
 int32 ConvertSerializedPayloadFlagsToInt32(const SerializedDataT& Data)
 {
-	return static_cast<int32>(Data.PayloadCompressionDetails);
+	FCompressionMethodAndBias MethodAndBias(Data.PayloadCompressionMethod, Data.PayloadCompressionBias);
+	return MethodAndBias.ToInt32();
 }
 
 } // namespace PackageDataUtil
diff --git a/Engine/Plugins/Developer/Concert/ConcertSync/ConcertSyncServer/Source/ConcertSyncServer/Private/ConcertServerDataStore.cpp b/Engine/Plugins/Developer/Concert/ConcertSync/ConcertSyncServer/Source/ConcertSyncServer/Private/ConcertServerDataStore.cpp
--- a/Engine/Plugins/Developer/Concert/ConcertSync/ConcertSyncServer/Source/ConcertSyncServer/Private/ConcertServerDataStore.cpp
+++ b/Engine/Plugins/Developer/Concert/ConcertSync/ConcertSyncServer/Source/ConcertSyncServer/Private/ConcertServerDataStore.cpp
@@ -67,7 +67,8 @@ EConcertSessionResponseCode FConcertServerDataStore::OnCompareExchange(const FCo
 		// If the version or value matches.
 		if ((Request.ExpectedVersion != 0 && Request.ExpectedVersion == Result.Value->Version) ||
 		    (Request.ExpectedVersion == 0
-			 && Result.Value->SerializedValue.PayloadCompressionDetails == Request.Expected.PayloadCompressionDetails
+			 && Result.Value->SerializedValue.PayloadCompressionMethod == Request.Expected.PayloadCompressionMethod
+			 && Result.Value->SerializedValue.PayloadCompressionBias == Request.Expected.PayloadCompressionBias
 			 && Result.Value->SerializedValue.PayloadSize == Request.Expected.PayloadSize
 			 && Result.Value->SerializedValue.PayloadBytes.Bytes == Request.Expected.PayloadBytes.Bytes))
 		{
